Just putting the boiler plate in palce so that we can start the work on this.

#Synced-CL 50972958

diff --git a/Engine/Plugins/Experimental/SemanticSearch/SemanticSearch.uplugin b/Engine/Plugins/Experimental/SemanticSearch/SemanticSearch.uplugin
new file mode 100644
--- /dev/null
+++ b/Engine/Plugins/Experimental/SemanticSearch/SemanticSearch.uplugin
@@ -0,0 +1,25 @@
+{
+	"FileVersion": 3,
+	"Version": 1,
+	"VersionName": "0.1",
+	"FriendlyName": "Semantic Search",
+	"Description": "Very early work in progress of a semantic search system for assets",
+	"Category": "Other",
+	"CreatedBy": "Epic Games, Inc.",
+	"CreatedByURL": "https://epicgames.com",
+	"DocsURL": "",
+	"MarketplaceURL": "",
+	"SupportURL": "https://forums.unrealengine.com/",
+	"CanContainContent": false,
+	"IsBetaVersion": false,
+	"IsExperimentalVersion": true,
+	"Installed": false,
+	"EnabledByDefault": false,
+	"Modules": [
+		{
+			"Name": "SemanticSearch",
+			"Type": "Editor",
+			"LoadingPhase": "Default"
+		}
+	]
+}
\ No newline at end of file
diff --git a/Engine/Plugins/Experimental/SemanticSearch/Source/SemanticSearch/Private/SemanticSearchModule.cpp b/Engine/Plugins/Experimental/SemanticSearch/Source/SemanticSearch/Private/SemanticSearchModule.cpp
new file mode 100644
--- /dev/null
+++ b/Engine/Plugins/Experimental/SemanticSearch/Source/SemanticSearch/Private/SemanticSearchModule.cpp
@@ -0,0 +1,7 @@
+// Copyright Epic Games, Inc. All Rights Reserved.
+
+#include "SemanticSearchModule.h"
+
+#include "Modules/ModuleManager.h"
+
+IMPLEMENT_MODULE(FSemanticSearchModule, SemanticSearch)
diff --git a/Engine/Plugins/Experimental/SemanticSearch/Source/SemanticSearch/Private/SemanticSearchModule.h b/Engine/Plugins/Experimental/SemanticSearch/Source/SemanticSearch/Private/SemanticSearchModule.h
new file mode 100644
--- /dev/null
+++ b/Engine/Plugins/Experimental/SemanticSearch/Source/SemanticSearch/Private/SemanticSearchModule.h
@@ -0,0 +1,11 @@
+// Copyright Epic Games, Inc. All Rights Reserved.
+
+#pragma once
+
+
+#include "Modules/ModuleInterface.h"
+
+class FSemanticSearchModule : public IModuleInterface
+{
+
+};
\ No newline at end of file
diff --git a/Engine/Plugins/Experimental/SemanticSearch/Source/SemanticSearch/SemanticSearch.Build.cs b/Engine/Plugins/Experimental/SemanticSearch/Source/SemanticSearch/SemanticSearch.Build.cs
new file mode 100644
--- /dev/null
+++ b/Engine/Plugins/Experimental/SemanticSearch/Source/SemanticSearch/SemanticSearch.Build.cs
@@ -0,0 +1,24 @@
+// Copyright Epic Games, Inc. All Rights Reserved.
+
+using UnrealBuildTool;
+
+/**
+ * Core module for semantic search service
+ * Provide the key system for the queries and management of the generation of the embeddings
+ * 
+ * Any ui or higher level logic should go in another module
+ */
+public class SemanticSearch : ModuleRules
+{
+	public SemanticSearch(ReadOnlyTargetRules Target) : base(Target)
+	{
+		PrivateDependencyModuleNames.AddRange(
+			new string[]
+			{
+				"AssetRegistry",
+				"Core",
+				"CoreUObject",
+			}
+			);
+	}
+}
\ No newline at end of file
