[Semantic Search] Add a RemoteEmbeddingProvider which calls temp captioning and embedding endpoints

#Synced-CL 51056086

#Synced-CL 51056086

diff --git a/Engine/Plugins/Experimental/SemanticSearch/Source/SemanticSearch/Private/RemoteEmbeddingProvider.cpp b/Engine/Plugins/Experimental/SemanticSearch/Source/SemanticSearch/Private/RemoteEmbeddingProvider.cpp
new file mode 100644
--- /dev/null
+++ b/Engine/Plugins/Experimental/SemanticSearch/Source/SemanticSearch/Private/RemoteEmbeddingProvider.cpp
@@ -0,0 +1,296 @@
+// Copyright Epic Games, Inc. All Rights Reserved.
+
+#include "RemoteEmbeddingProvider.h"
+#include "SemanticSearchModule.h"
+#include "Serialization/JsonReader.h"
+#include "Serialization/JsonSerializer.h"
+#include "Policies/CondensedJsonPrintPolicy.h"
+#include "Dom/JsonObject.h"
+#include "Misc/FileHelper.h"
+#include "Misc/Paths.h"
+
+static FString ReadBaseURLFromConfig()
+{
+	const FString ConfigPath = FPaths::EngineDir() / TEXT("Restricted/NotForLicensees/Config/SemanticSearch.json");
+	FString FileContents;
+	if (!FFileHelper::LoadFileToString(FileContents, *ConfigPath))
+	{
+		UE_LOG(LogSemanticSearch, Warning, TEXT("Failed to load config: %s"), *ConfigPath);
+		return FString();
+	}
+
+	TSharedPtr<FJsonObject> Json;
+	TSharedRef<TJsonReader<>> Reader = TJsonReaderFactory<>::Create(FileContents);
+	if (!FJsonSerializer::Deserialize(Reader, Json) || !Json.IsValid())
+	{
+		UE_LOG(LogSemanticSearch, Warning, TEXT("Failed to parse config JSON: %s"), *ConfigPath);
+		return FString();
+	}
+
+	FString URL;
+	if (!Json->TryGetStringField(TEXT("base_url"), URL))
+	{
+		UE_LOG(LogSemanticSearch, Warning, TEXT("Missing 'base_url' in config: %s"), *ConfigPath);
+	}
+	return URL;
+}
+
+FRemoteEmbeddingProvider::FRemoteEmbeddingProvider(const FString& InBaseURL)
+	: BaseURL(InBaseURL.IsEmpty() ? ReadBaseURLFromConfig() : InBaseURL)
+	, EmbeddingDimension(0)
+{
+	UE_LOG(LogSemanticSearch, Log, TEXT("RemoteEmbeddingProvider created (URL: %s)"), *BaseURL);
+}
+
+int32 FRemoteEmbeddingProvider::GetEmbeddingDimension() const
+{
+	return EmbeddingDimension.Load();
+}
+
+void FRemoteEmbeddingProvider::GenerateCaptionAsync(
+	const FCaptionRequest& Request,
+	FOnCaptionComplete OnComplete)
+{
+	FString Boundary;
+	TArray<uint8> Body = BuildMultipartBody(Request, Boundary);
+
+	TSharedRef<IHttpRequest, ESPMode::ThreadSafe> Http = FHttpModule::Get().CreateRequest();
+	Http->SetURL(BaseURL / TEXT("v1") / TEXT("caption"));
+	Http->SetVerb(TEXT("POST"));
+	Http->SetHeader(TEXT("Content-Type"),
+		FString::Printf(TEXT("multipart/form-data; boundary=%s"), *Boundary));
+	Http->SetContent(MoveTemp(Body));
+	Http->SetTimeout(60.0f);
+
+	Http->OnProcessRequestComplete().BindLambda(
+		[this, OnComplete](FHttpRequestPtr /*Req*/, FHttpResponsePtr Resp, bool bConnected)
+		{
+			FCaptionResponse Response;
+
+			if (!bConnected || !Resp.IsValid())
+			{
+				Response.ErrorMessage = TEXT("Connection failed");
+				OnComplete.ExecuteIfBound(Response);
+				return;
+			}
+
+			if (Resp->GetResponseCode() != 200)
+			{
+				Response.ErrorMessage = FString::Printf(TEXT("HTTP %d: %s"),
+					Resp->GetResponseCode(), *Resp->GetContentAsString().Left(200));
+				OnComplete.ExecuteIfBound(Response);
+				return;
+			}
+
+			TSharedPtr<FJsonObject> Json;
+			TSharedRef<TJsonReader<>> Reader = TJsonReaderFactory<>::Create(Resp->GetContentAsString());
+			if (!FJsonSerializer::Deserialize(Reader, Json) || !Json.IsValid())
+			{
+				Response.ErrorMessage = TEXT("Failed to parse JSON response");
+				OnComplete.ExecuteIfBound(Response);
+				return;
+			}
+
+			// Server-side error
+			FString ServerError;
+			if (Json->TryGetStringField(TEXT("error"), ServerError))
+			{
+				Response.ErrorMessage = FString::Printf(TEXT("Server error: %s"), *ServerError);
+				OnComplete.ExecuteIfBound(Response);
+				return;
+			}
+
+			// Caption
+			if (!Json->TryGetStringField(TEXT("caption"), Response.Caption))
+			{
+				Response.ErrorMessage = TEXT("Missing 'caption' in response");
+				OnComplete.ExecuteIfBound(Response);
+				return;
+			}
+
+			// Keywords
+			const TArray<TSharedPtr<FJsonValue>>* KeywordsArray;
+			if (Json->TryGetArrayField(TEXT("keywords"), KeywordsArray))
+			{
+				for (const TSharedPtr<FJsonValue>& Val : *KeywordsArray)
+				{
+					FString Kw;
+					if (Val->TryGetString(Kw))
+					{
+						Response.Keywords.Add(Kw);
+					}
+				}
+			}
+
+			OnComplete.ExecuteIfBound(Response);
+		});
+
+	Http->ProcessRequest();
+}
+
+void FRemoteEmbeddingProvider::GenerateEmbeddingAsync(
+	const FString& Text,
+	FOnEmbeddingComplete OnComplete)
+{
+	TSharedRef<FJsonObject> Payload = MakeShared<FJsonObject>();
+	Payload->SetStringField(TEXT("text"), Text);
+
+	FString PayloadStr;
+	TSharedRef<TJsonWriter<>> Writer = TJsonWriterFactory<>::Create(&PayloadStr);
+	FJsonSerializer::Serialize(Payload, Writer);
+
+	TSharedRef<IHttpRequest, ESPMode::ThreadSafe> Http = FHttpModule::Get().CreateRequest();
+	Http->SetURL(BaseURL / TEXT("v1") / TEXT("embed"));
+	Http->SetVerb(TEXT("POST"));
+	Http->SetHeader(TEXT("Content-Type"), TEXT("application/json"));
+	Http->SetContentAsString(PayloadStr);
+	Http->SetTimeout(30.0f);
+
+	Http->OnProcessRequestComplete().BindLambda(
+		[this, OnComplete](FHttpRequestPtr /*Req*/, FHttpResponsePtr Resp, bool bConnected)
+		{
+			FEmbeddingResponse Response;
+
+			if (!bConnected || !Resp.IsValid())
+			{
+				Response.ErrorMessage = TEXT("Connection failed");
+				OnComplete.ExecuteIfBound(Response);
+				return;
+			}
+
+			if (Resp->GetResponseCode() != 200)
+			{
+				Response.ErrorMessage = FString::Printf(TEXT("HTTP %d: %s"),
+					Resp->GetResponseCode(), *Resp->GetContentAsString().Left(200));
+				OnComplete.ExecuteIfBound(Response);
+				return;
+			}
+
+			TSharedPtr<FJsonObject> Json;
+			TSharedRef<TJsonReader<>> Reader = TJsonReaderFactory<>::Create(Resp->GetContentAsString());
+			if (!FJsonSerializer::Deserialize(Reader, Json) || !Json.IsValid())
+			{
+				Response.ErrorMessage = TEXT("Failed to parse JSON response");
+				OnComplete.ExecuteIfBound(Response);
+				return;
+			}
+
+			// Server-side error
+			FString ServerError;
+			if (Json->TryGetStringField(TEXT("error"), ServerError))
+			{
+				Response.ErrorMessage = FString::Printf(TEXT("Server error: %s"), *ServerError);
+				OnComplete.ExecuteIfBound(Response);
+				return;
+			}
+
+			// Embedding array
+			const TArray<TSharedPtr<FJsonValue>>* EmbArray;
+			if (!Json->TryGetArrayField(TEXT("embedding"), EmbArray))
+			{
+				Response.ErrorMessage = TEXT("Missing 'embedding' in response");
+				OnComplete.ExecuteIfBound(Response);
+				return;
+			}
+
+			Response.Embedding.Reserve(EmbArray->Num());
+			for (const TSharedPtr<FJsonValue>& Val : *EmbArray)
+			{
+				double Num;
+				if (Val->TryGetNumber(Num))
+				{
+					Response.Embedding.Add(static_cast<float>(Num));
+				}
+				else
+				{
+					Response.Embedding.Empty();
+					Response.ErrorMessage = TEXT("Invalid value in embedding array");
+					OnComplete.ExecuteIfBound(Response);
+					return;
+				}
+			}
+
+			// Auto-detect dimension from first successful response
+			if (EmbeddingDimension.Load() == 0 && Response.Embedding.Num() > 0)
+			{
+				EmbeddingDimension.Store(Response.Embedding.Num());
+				UE_LOG(LogSemanticSearch, Log, TEXT("RemoteProvider: Embedding dimension: %d"),
+					Response.Embedding.Num());
+			}
+
+			// Model version
+			Json->TryGetStringField(TEXT("model_version"), Response.ModelVersion);
+
+			OnComplete.ExecuteIfBound(Response);
+		});
+
+	Http->ProcessRequest();
+}
+
+TArray<uint8> FRemoteEmbeddingProvider::BuildMultipartBody(
+	const FCaptionRequest& Request, FString& OutBoundary) const
+{
+	OutBoundary = FString::Printf(TEXT("----UE%s"), *FGuid::NewGuid().ToString(EGuidFormats::Digits));
+
+	TArray<uint8> Body;
+	const FString Delim = FString::Printf(TEXT("--%s\r\n"), *OutBoundary);
+	const FString End   = FString::Printf(TEXT("--%s--\r\n"), *OutBoundary);
+
+	auto Append = [&Body](const FString& S)
+	{
+		FTCHARToUTF8 Utf8(*S);
+		Body.Append(reinterpret_cast<const uint8*>(Utf8.Get()), Utf8.Length());
+	};
+
+	// media[] file parts + media_types[] form fields
+	for (int32 i = 0; i < Request.AssetMedia.Num(); ++i)
+	{
+		const FAssetMedia& M = Request.AssetMedia[i];
+		if (M.Data.Num() == 0) continue;
+
+		FString Ext = TEXT("bin");
+		if      (M.MimeType == TEXT("image/png"))          Ext = TEXT("png");
+		else if (M.MimeType == TEXT("image/jpeg"))         Ext = TEXT("jpg");
+		else if (M.MimeType == TEXT("application/json"))   Ext = TEXT("json");
+		else if (M.MimeType.StartsWith(TEXT("text/")))     Ext = TEXT("txt");
+
+		Append(Delim);
+		Append(FString::Printf(TEXT("Content-Disposition: form-data; name=\"media[]\"; filename=\"media_%d.%s\"\r\n"), i, *Ext));
+		Append(FString::Printf(TEXT("Content-Type: %s\r\n\r\n"), *M.MimeType));
+		Body.Append(M.Data.GetData(), M.Data.Num());
+		Append(TEXT("\r\n"));
+
+		Append(Delim);
+		Append(TEXT("Content-Disposition: form-data; name=\"media_types[]\"\r\n\r\n"));
+		Append(FString::Printf(TEXT("%s\r\n"), *M.MimeType));
+	}
+
+	if (!Request.AssetPath.IsEmpty())
+	{
+		Append(Delim);
+		Append(TEXT("Content-Disposition: form-data; name=\"asset_path\"\r\n\r\n"));
+		Append(FString::Printf(TEXT("%s\r\n"), *Request.AssetPath));
+	}
+
+	if (!Request.AssetType.IsEmpty())
+	{
+		Append(Delim);
+		Append(TEXT("Content-Disposition: form-data; name=\"asset_class\"\r\n\r\n"));
+		Append(FString::Printf(TEXT("%s\r\n"), *Request.AssetType));
+	}
+
+	if (Request.Metadata.IsValid() && Request.Metadata->Values.Num() > 0)
+	{
+		FString MetaStr;
+		TSharedRef<TJsonWriter<TCHAR, TCondensedJsonPrintPolicy<TCHAR>>> Wr =
+			TJsonWriterFactory<TCHAR, TCondensedJsonPrintPolicy<TCHAR>>::Create(&MetaStr);
+		FJsonSerializer::Serialize(Request.Metadata.ToSharedRef(), Wr);
+
+		Append(Delim);
+		Append(TEXT("Content-Disposition: form-data; name=\"metadata\"\r\n\r\n"));
+		Append(FString::Printf(TEXT("%s\r\n"), *MetaStr));
+	}
+
+	Append(End);
+	return Body;
+}
diff --git a/Engine/Plugins/Experimental/SemanticSearch/Source/SemanticSearch/Private/RemoteEmbeddingProvider.h b/Engine/Plugins/Experimental/SemanticSearch/Source/SemanticSearch/Private/RemoteEmbeddingProvider.h
new file mode 100644
--- /dev/null
+++ b/Engine/Plugins/Experimental/SemanticSearch/Source/SemanticSearch/Private/RemoteEmbeddingProvider.h
@@ -0,0 +1,35 @@
+// Copyright Epic Games, Inc. All Rights Reserved.
+
+#pragma once
+
+#include "Interfaces/IEmbeddingProvider.h"
+#include "HttpModule.h"
+#include "Interfaces/IHttpRequest.h"
+#include "Interfaces/IHttpResponse.h"
+
+/**
+ * Remote embedding provider that communicates with a Python backend.
+ *
+ *   GenerateCaptionAsync   -> POST /v1/caption  (multipart/form-data)
+ *   GenerateEmbeddingAsync -> POST /v1/embed    (application/json)
+ */
+class FRemoteEmbeddingProvider : public IEmbeddingProvider
+{
+public:
+	/** If no URL provided, reads from Engine/Restricted/NotForLicensees/Config/SemanticSearch.json */
+	explicit FRemoteEmbeddingProvider(const FString& InBaseURL = FString());
+
+	//--- IEmbeddingProvider ---
+	virtual int32 GetEmbeddingDimension() const override;
+	virtual void GenerateCaptionAsync(const FCaptionRequest& Request, FOnCaptionComplete OnComplete) override;
+	virtual void GenerateEmbeddingAsync(const FString& Text, FOnEmbeddingComplete OnComplete) override;
+
+	FString GetBaseURL() const { return BaseURL; }
+
+private:
+	/** Build multipart/form-data body. Sets OutBoundary for Content-Type header. */
+	TArray<uint8> BuildMultipartBody(const FCaptionRequest& Request, FString& OutBoundary) const;
+
+	FString BaseURL;
+	TAtomic<int32> EmbeddingDimension;
+};
diff --git a/Engine/Plugins/Experimental/SemanticSearch/Source/SemanticSearch/Private/SemanticSearchModule.cpp b/Engine/Plugins/Experimental/SemanticSearch/Source/SemanticSearch/Private/SemanticSearchModule.cpp
--- a/Engine/Plugins/Experimental/SemanticSearch/Source/SemanticSearch/Private/SemanticSearchModule.cpp
+++ b/Engine/Plugins/Experimental/SemanticSearch/Source/SemanticSearch/Private/SemanticSearchModule.cpp
@@ -1,7 +1,17 @@
 // Copyright Epic Games, Inc. All Rights Reserved.
 
 #include "SemanticSearchModule.h"
-
 #include "Modules/ModuleManager.h"
 
+DEFINE_LOG_CATEGORY(LogSemanticSearch);
+
 IMPLEMENT_MODULE(FSemanticSearchModule, SemanticSearch)
+void FSemanticSearchModule::StartupModule()
+{
+	UE_LOG(LogSemanticSearch, Log, TEXT("SemanticSearch module starting up"));
+}
+
+void FSemanticSearchModule::ShutdownModule()
+{
+	UE_LOG(LogSemanticSearch, Log, TEXT("SemanticSearch module shutting down"));
+}
diff --git a/Engine/Plugins/Experimental/SemanticSearch/Source/SemanticSearch/Private/SemanticSearchModule.h b/Engine/Plugins/Experimental/SemanticSearch/Source/SemanticSearch/Private/SemanticSearchModule.h
deleted file mode 100644
diff --git a/Engine/Plugins/Experimental/SemanticSearch/Source/SemanticSearch/Public/Interfaces/IEmbeddingProvider.h b/Engine/Plugins/Experimental/SemanticSearch/Source/SemanticSearch/Public/Interfaces/IEmbeddingProvider.h
new file mode 100644
--- /dev/null
+++ b/Engine/Plugins/Experimental/SemanticSearch/Source/SemanticSearch/Public/Interfaces/IEmbeddingProvider.h
@@ -0,0 +1,109 @@
+// Copyright Epic Games, Inc. All Rights Reserved.
+
+#pragma once
+
+#include "CoreMinimal.h"
+#include "Dom/JsonObject.h"
+
+/**
+ * A single media payload attached to a caption request.
+ */
+struct SEMANTICSEARCH_API FAssetMedia
+{
+	/** MIME type: "image/png", "application/json", "audio/mpeg", etc. */
+	FString MimeType;
+
+	/** Raw media bytes */
+	TArray<uint8> Data;
+};
+
+/**
+ * Request to generate a caption (and keywords) for an asset.
+ * Can include media content and structured metadata.
+ */
+struct SEMANTICSEARCH_API FCaptionRequest
+{
+	FString AssetPath;
+	FString AssetType;
+
+	/** Asset content — can be combination of images, text, JSON */
+	TArray<FAssetMedia> AssetMedia;
+
+	/** Structured metadata (asset dimension, properties, etc.) */
+	TSharedPtr<FJsonObject> Metadata;
+
+	/** Content hash for DDC cache key generation */
+	FString ContentHash;
+};
+
+/**
+ * Response from caption generation.
+ * Success if ErrorMessage is empty.
+ */
+struct SEMANTICSEARCH_API FCaptionResponse
+{
+	FString ErrorMessage;
+
+	FString Caption;
+
+	TArray<FString> Keywords;
+};
+
+/**
+ * Response from embedding generation.
+ * Success if ErrorMessage is empty.
+ */
+struct SEMANTICSEARCH_API FEmbeddingResponse
+{
+	FString ErrorMessage;
+
+	TArray<float> Embedding;
+
+	/** Version of the model that generated this embedding */
+	FString ModelVersion;
+};
+
+/**
+ * Delegate for async caption completion
+ */
+DECLARE_DELEGATE_OneParam(FOnCaptionComplete, const FCaptionResponse&);
+
+/**
+ * Delegate for async embedding completion
+ */
+DECLARE_DELEGATE_OneParam(FOnEmbeddingComplete, const FEmbeddingResponse&);
+
+/**
+ * Interface for embedding generation providers.
+ * Can support both remote (HTTP API) and local (on-device) implementations.
+ */
+class SEMANTICSEARCH_API IEmbeddingProvider
+{
+public:
+	virtual ~IEmbeddingProvider() = default;
+
+	/** Embedding vector dimensionality */
+	virtual int32 GetEmbeddingDimension() const = 0;
+
+	/**
+	 * Generate a caption and keywords for an asset.
+	 * Called during indexing to describe the asset content.
+	 *
+	 * @param Request The caption request with asset media and metadata
+	 * @param OnComplete Callback invoked when caption is ready
+	 */
+	virtual void GenerateCaptionAsync(
+		const FCaptionRequest& Request,
+		FOnCaptionComplete OnComplete) = 0;
+
+	/**
+	 * Generate embedding for a text string.
+	 * Called during indexing (with caption text) and at search time (with query text).
+	 *
+	 * @param Text The text to embed (caption or search query)
+	 * @param OnComplete Callback invoked when embedding is ready
+	 */
+	virtual void GenerateEmbeddingAsync(
+		const FString& Text,
+		FOnEmbeddingComplete OnComplete) = 0;
+};
diff --git a/Engine/Plugins/Experimental/SemanticSearch/Source/SemanticSearch/Public/SemanticSearchModule.h b/Engine/Plugins/Experimental/SemanticSearch/Source/SemanticSearch/Public/SemanticSearchModule.h
new file mode 100644
--- /dev/null
+++ b/Engine/Plugins/Experimental/SemanticSearch/Source/SemanticSearch/Public/SemanticSearchModule.h
@@ -0,0 +1,17 @@
+// Copyright Epic Games, Inc. All Rights Reserved.
+
+#pragma once
+
+
+#include "CoreMinimal.h"
+#include "Modules/ModuleInterface.h"
+
+DECLARE_LOG_CATEGORY_EXTERN(LogSemanticSearch, Log, All);
+
+class FSemanticSearchModule : public IModuleInterface
+{
+	public:
+
+	virtual void StartupModule() override;
+	virtual void ShutdownModule() override;
+};
\ No newline at end of file
diff --git a/Engine/Plugins/Experimental/SemanticSearch/Source/SemanticSearch/SemanticSearch.Build.cs b/Engine/Plugins/Experimental/SemanticSearch/Source/SemanticSearch/SemanticSearch.Build.cs
--- a/Engine/Plugins/Experimental/SemanticSearch/Source/SemanticSearch/SemanticSearch.Build.cs
+++ b/Engine/Plugins/Experimental/SemanticSearch/Source/SemanticSearch/SemanticSearch.Build.cs
@@ -18,6 +18,8 @@ public class SemanticSearch : ModuleRules
 				"AssetRegistry",
 				"Core",
 				"CoreUObject",
+				"HTTP",
+				"Json",
 			}
 			);
 	}
