Add test commands and config file

#Synced-CL 51056086

diff --git a/Engine/Plugins/Experimental/SemanticSearch/Source/SemanticSearch/Private/RemoteEmbeddingProvider.cpp b/Engine/Plugins/Experimental/SemanticSearch/Source/SemanticSearch/Private/RemoteEmbeddingProvider.cpp
--- a/Engine/Plugins/Experimental/SemanticSearch/Source/SemanticSearch/Private/RemoteEmbeddingProvider.cpp
+++ b/Engine/Plugins/Experimental/SemanticSearch/Source/SemanticSearch/Private/RemoteEmbeddingProvider.cpp
@@ -91,7 +91,6 @@ void FRemoteEmbeddingProvider::GenerateCaptionAsync(
 				return;
 			}
 
-			// Server-side error
 			FString ServerError;
 			if (Json->TryGetStringField(TEXT("error"), ServerError))
 			{
@@ -100,7 +99,6 @@ void FRemoteEmbeddingProvider::GenerateCaptionAsync(
 				return;
 			}
 
-			// Caption
 			if (!Json->TryGetStringField(TEXT("caption"), Response.Caption))
 			{
 				Response.ErrorMessage = TEXT("Missing 'caption' in response");
@@ -108,7 +106,6 @@ void FRemoteEmbeddingProvider::GenerateCaptionAsync(
 				return;
 			}
 
-			// Keywords
 			const TArray<TSharedPtr<FJsonValue>>* KeywordsArray;
 			if (Json->TryGetArrayField(TEXT("keywords"), KeywordsArray))
 			{
@@ -175,7 +172,6 @@ void FRemoteEmbeddingProvider::GenerateEmbeddingAsync(
 				return;
 			}
 
-			// Server-side error
 			FString ServerError;
 			if (Json->TryGetStringField(TEXT("error"), ServerError))
 			{
@@ -184,7 +180,6 @@ void FRemoteEmbeddingProvider::GenerateEmbeddingAsync(
 				return;
 			}
 
-			// Embedding array
 			const TArray<TSharedPtr<FJsonValue>>* EmbArray;
 			if (!Json->TryGetArrayField(TEXT("embedding"), EmbArray))
 			{
@@ -210,7 +205,6 @@ void FRemoteEmbeddingProvider::GenerateEmbeddingAsync(
 				}
 			}
 
-			// Auto-detect dimension from first successful response
 			if (EmbeddingDimension.Load() == 0 && Response.Embedding.Num() > 0)
 			{
 				EmbeddingDimension.Store(Response.Embedding.Num());
@@ -218,7 +212,6 @@ void FRemoteEmbeddingProvider::GenerateEmbeddingAsync(
 					Response.Embedding.Num());
 			}
 
-			// Model version
 			Json->TryGetStringField(TEXT("model_version"), Response.ModelVersion);
 
 			OnComplete.ExecuteIfBound(Response);
diff --git a/Engine/Plugins/Experimental/SemanticSearch/Source/SemanticSearch/Private/RemoteEmbeddingProvider.h b/Engine/Plugins/Experimental/SemanticSearch/Source/SemanticSearch/Private/RemoteEmbeddingProvider.h
--- a/Engine/Plugins/Experimental/SemanticSearch/Source/SemanticSearch/Private/RemoteEmbeddingProvider.h
+++ b/Engine/Plugins/Experimental/SemanticSearch/Source/SemanticSearch/Private/RemoteEmbeddingProvider.h
@@ -8,18 +8,16 @@
 #include "Interfaces/IHttpResponse.h"
 
 /**
- * Remote embedding provider that communicates with a Python backend.
+ *  Remote embedding provider that communicates with Python backend.
  *
- *   GenerateCaptionAsync   -> POST /v1/caption  (multipart/form-data)
- *   GenerateEmbeddingAsync -> POST /v1/embed    (application/json)
+ *  GenerateCaptionAsync   -> POST /v1/caption  (multipart/form-data)
+ *  GenerateEmbeddingAsync -> POST /v1/embed    (application/json)
  */
 class FRemoteEmbeddingProvider : public IEmbeddingProvider
 {
 public:
-	/** If no URL provided, reads from Engine/Restricted/NotForLicensees/Config/SemanticSearch.json */
 	explicit FRemoteEmbeddingProvider(const FString& InBaseURL = FString());
 
-	//--- IEmbeddingProvider ---
 	virtual int32 GetEmbeddingDimension() const override;
 	virtual void GenerateCaptionAsync(const FCaptionRequest& Request, FOnCaptionComplete OnComplete) override;
 	virtual void GenerateEmbeddingAsync(const FString& Text, FOnEmbeddingComplete OnComplete) override;
@@ -27,7 +25,7 @@ public:
 	FString GetBaseURL() const { return BaseURL; }
 
 private:
-	/** Build multipart/form-data body. Sets OutBoundary for Content-Type header. */
+	/** Build multipart/form-data body.*/
 	TArray<uint8> BuildMultipartBody(const FCaptionRequest& Request, FString& OutBoundary) const;
 
 	FString BaseURL;
diff --git a/Engine/Plugins/Experimental/SemanticSearch/Source/SemanticSearch/Private/SemanticSearchTestCommands.cpp b/Engine/Plugins/Experimental/SemanticSearch/Source/SemanticSearch/Private/SemanticSearchTestCommands.cpp
new file mode 100644
--- /dev/null
+++ b/Engine/Plugins/Experimental/SemanticSearch/Source/SemanticSearch/Private/SemanticSearchTestCommands.cpp
@@ -0,0 +1,126 @@
+// Copyright Epic Games, Inc. All Rights Reserved.
+
+#include "SemanticSearchModule.h"
+#include "RemoteEmbeddingProvider.h"
+#include "HAL/ConsoleManager.h"
+#include "Misc/FileHelper.h"
+#include "Misc/Paths.h"
+
+static TSharedPtr<FRemoteEmbeddingProvider> GTestProvider;
+
+static TSharedPtr<FRemoteEmbeddingProvider> GetOrCreateProvider()
+{
+	if (!GTestProvider.IsValid())
+	{
+		GTestProvider = MakeShared<FRemoteEmbeddingProvider>();
+	}
+	return GTestProvider;
+}
+
+// SemanticSearch.TestCaption <image_file_path>
+static FAutoConsoleCommand TestCaptionCommand(
+	TEXT("SemanticSearch.TestCaption"),
+	TEXT("Test caption generation from an image file. Usage: SemanticSearch.TestCaption <image_file_path>"),
+	FConsoleCommandWithArgsDelegate::CreateLambda([](const TArray<FString>& Args)
+	{
+		if (Args.Num() < 1)
+		{
+			UE_LOG(LogSemanticSearch, Warning, TEXT("Usage: SemanticSearch.TestCaption <image_file_path>"));
+			return;
+		}
+
+		const FString FilePath = Args[0];
+		if (!FPaths::FileExists(FilePath))
+		{
+			UE_LOG(LogSemanticSearch, Error, TEXT("File not found: %s"), *FilePath);
+			return;
+		}
+
+		TArray<uint8> FileData;
+		if (!FFileHelper::LoadFileToArray(FileData, *FilePath))
+		{
+			UE_LOG(LogSemanticSearch, Error, TEXT("Failed to load file: %s"), *FilePath);
+			return;
+		}
+
+		const FString Ext = FPaths::GetExtension(FilePath).ToLower();
+		FString MimeType = TEXT("application/octet-stream");
+		if (Ext == TEXT("png"))              MimeType = TEXT("image/png");
+		else if (Ext == TEXT("jpg") || Ext == TEXT("jpeg")) MimeType = TEXT("image/jpeg");
+		else if (Ext == TEXT("bmp"))         MimeType = TEXT("image/bmp");
+
+		FCaptionRequest Request;
+		FAssetMedia Media;
+		Media.MimeType = MimeType;
+		Media.Data = MoveTemp(FileData);
+		Request.AssetMedia.Add(MoveTemp(Media));
+
+		UE_LOG(LogSemanticSearch, Log, TEXT("Sending caption request (%s, %d bytes)..."), *MimeType, Request.AssetMedia[0].Data.Num());
+
+		TSharedPtr<FRemoteEmbeddingProvider> Provider = GetOrCreateProvider();
+		Provider->GenerateCaptionAsync(Request,
+			FOnCaptionComplete::CreateLambda([](const FCaptionResponse& Response)
+			{
+				if (!Response.ErrorMessage.IsEmpty())
+				{
+					UE_LOG(LogSemanticSearch, Error, TEXT("Caption failed: %s"), *Response.ErrorMessage);
+					return;
+				}
+
+				UE_LOG(LogSemanticSearch, Log, TEXT("Caption: %s"), *Response.Caption);
+				if (Response.Keywords.Num() > 0)
+				{
+					UE_LOG(LogSemanticSearch, Log, TEXT("Keywords: %s"), *FString::Join(Response.Keywords, TEXT(", ")));
+				}
+			}));
+	})
+);
+
+// SemanticSearch.TestEmbedding [text]
+static FAutoConsoleCommand TestEmbeddingCommand(
+	TEXT("SemanticSearch.TestEmbedding"),
+	TEXT("Test embedding generation. Usage: SemanticSearch.TestEmbedding [text to embed]"),
+	FConsoleCommandWithArgsDelegate::CreateLambda([](const TArray<FString>& Args)
+	{
+		FString Text;
+		if (Args.Num() > 0)
+		{
+			Text = FString::Join(Args, TEXT(" "));
+		}
+		else
+		{
+			Text = TEXT("A red wooden chair with four legs");
+		}
+
+		UE_LOG(LogSemanticSearch, Log, TEXT("Generating embedding for: \"%s\""), *Text);
+
+		TSharedPtr<FRemoteEmbeddingProvider> Provider = GetOrCreateProvider();
+		Provider->GenerateEmbeddingAsync(Text,
+			FOnEmbeddingComplete::CreateLambda([](const FEmbeddingResponse& Response)
+			{
+				if (!Response.ErrorMessage.IsEmpty())
+				{
+					UE_LOG(LogSemanticSearch, Error, TEXT("Embedding failed: %s"), *Response.ErrorMessage);
+					return;
+				}
+
+				UE_LOG(LogSemanticSearch, Log, TEXT("Embedding: %d dimensions"), Response.Embedding.Num());
+				if (!Response.ModelVersion.IsEmpty())
+				{
+					UE_LOG(LogSemanticSearch, Log, TEXT("Model: %s"), *Response.ModelVersion);
+				}
+
+				if (Response.Embedding.Num() > 0)
+				{
+					const int32 Preview = FMath::Min(Response.Embedding.Num(), 5);
+					FString Vals;
+					for (int32 i = 0; i < Preview; ++i)
+					{
+						if (i > 0) Vals += TEXT(", ");
+						Vals += FString::Printf(TEXT("%.6f"), Response.Embedding[i]);
+					}
+					UE_LOG(LogSemanticSearch, Log, TEXT("First %d values: [%s, ...]"), Preview, *Vals);
+				}
+			}));
+	})
+);
