Sementic search added a first pass of the derived data job.Add the plumbing for the storing, fetching and generation of the data that will be shared acrross all asset types.

Reworked a bit the remote embbeding provider so that It can interface with the task owner system that the ddc use.

Things that still need to be done for sure here but not part of this cl. Logging and making the embbeding provider also part of the key generation for the derived data.

#Synced-CL 51416471

diff --git a/Engine/Plugins/Experimental/SemanticSearch/Source/SemanticSearch/Private/RemoteEmbeddingProvider.cpp b/Engine/Plugins/Experimental/SemanticSearch/Source/SemanticSearch/Private/RemoteEmbeddingProvider.cpp
--- a/Engine/Plugins/Experimental/SemanticSearch/Source/SemanticSearch/Private/RemoteEmbeddingProvider.cpp
+++ b/Engine/Plugins/Experimental/SemanticSearch/Source/SemanticSearch/Private/RemoteEmbeddingProvider.cpp
@@ -1,13 +1,95 @@
 // Copyright Epic Games, Inc. All Rights Reserved.
 
 #include "RemoteEmbeddingProvider.h"
-#include "SemanticSearchModule.h"
-#include "Serialization/JsonReader.h"
-#include "Serialization/JsonSerializer.h"
-#include "Policies/CondensedJsonPrintPolicy.h"
+
+#include "DerivedDataRequest.h"
+#include "DerivedDataRequestOwner.h"
+#include "DerivedDataRequestTypes.h"
 #include "Dom/JsonObject.h"
+#include "Interfaces/IHttpRequest.h"
 #include "Misc/FileHelper.h"
 #include "Misc/Paths.h"
+#include "Policies/CondensedJsonPrintPolicy.h"
+#include "SemanticSearchModule.h"
+#include "Serialization/JsonReader.h"
+#include "Serialization/JsonSerializer.h"
+
+namespace UE::SemanticSearch::Private
+{
+
+class FSemanticHttpRequest : public DerivedData::FRequestBase
+{
+public:
+	FSemanticHttpRequest(const TSharedRef<IHttpRequest, ESPMode::ThreadSafe>& InHttpRequest, DerivedData::FRequestOwner& InRequestOwner)
+		: FRequestBase()
+		, HttpRequest(InHttpRequest)
+		, RequestOwner(&static_cast<DerivedData::IRequestOwner&>(InRequestOwner))
+	{
+		RequestCompleteDelegate = MoveTemp(HttpRequest->OnProcessRequestComplete());
+		HttpRequest->OnProcessRequestComplete().BindRaw(this, &FSemanticHttpRequest::OnRequestComplete);
+		RequestOwner->Begin(this);
+		HttpRequest->ProcessRequest();
+	}
+
+	virtual ~FSemanticHttpRequest()
+	{
+		HttpRequest->OnProcessRequestComplete().Unbind();
+	}
+
+	virtual void SetPriority(DerivedData::EPriority Priority) override
+	{
+		EHttpRequestPriority HttpPriority = EHttpRequestPriority::Low;
+		switch (Priority)
+		{
+			case DerivedData::EPriority::Lowest:
+				HttpPriority = EHttpRequestPriority::Lowest;
+				break;
+			case DerivedData::EPriority::Low:
+				HttpPriority = EHttpRequestPriority::Low;
+				break;
+			case DerivedData::EPriority::Normal:
+				HttpPriority = EHttpRequestPriority::Normal;
+				break;
+			case DerivedData::EPriority::High:
+				HttpPriority = EHttpRequestPriority::High;
+				break;
+			case DerivedData::EPriority::Highest:
+			case DerivedData::EPriority::Blocking:
+				HttpPriority = EHttpRequestPriority::Highest;
+				break;
+
+			default:
+				break;
+		}
+
+		HttpRequest->SetPriority(HttpPriority);
+	}
+
+
+	virtual void Cancel() override
+	{
+		HttpRequest->CancelRequest();
+	}
+
+	virtual void Wait() override
+	{
+		HttpRequest->ProcessRequestUntilComplete();
+	}
+
+private:
+	void OnRequestComplete(FHttpRequestPtr Request, FHttpResponsePtr Response, bool bConnected)
+	{
+		RequestOwner->End(this, [this, &Request, &Response, bConnected]()
+			{
+				RequestCompleteDelegate.ExecuteIfBound(Request, Response, bConnected);
+			});
+	}
+
+	TSharedRef<IHttpRequest, ESPMode::ThreadSafe> HttpRequest;
+	DerivedData::IRequestOwner* RequestOwner;
+	FHttpRequestCompleteDelegate RequestCompleteDelegate;
+};
+
 
 static FString ReadBaseURLFromConfig()
 {
@@ -48,8 +130,9 @@ int32 FRemoteEmbeddingProvider::GetEmbeddingDimension() const
 }
 
 void FRemoteEmbeddingProvider::GenerateCaptionAsync(
-	const FCaptionRequest& Request,
-	FOnCaptionComplete OnComplete)
+	FCaptionRequest&& Request,
+	FOnCaptionComplete&& InOnComplete,
+	DerivedData::FRequestOwner* RequestOwner)
 {
 	FString Boundary;
 	TArray<uint8> Body = BuildMultipartBody(Request, Boundary);
@@ -63,14 +146,14 @@ void FRemoteEmbeddingProvider::GenerateCaptionAsync(
 	Http->SetTimeout(60.0f);
 
 	Http->OnProcessRequestComplete().BindLambda(
-		[this, OnComplete](FHttpRequestPtr /*Req*/, FHttpResponsePtr Resp, bool bConnected)
+		[this, OnComplete = MakeShared<FOnCaptionComplete>(MoveTemp(InOnComplete))](FHttpRequestPtr /*Req*/, FHttpResponsePtr Resp, bool bConnected)
 		{
 			FCaptionResponse Response;
 
 			if (!bConnected || !Resp.IsValid())
 			{
 				Response.ErrorMessage = TEXT("Connection failed");
-				OnComplete.ExecuteIfBound(Response);
+				OnComplete->operator()(MoveTemp(Response));
 				return;
 			}
 
@@ -78,7 +161,7 @@ void FRemoteEmbeddingProvider::GenerateCaptionAsync(
 			{
 				Response.ErrorMessage = FString::Printf(TEXT("HTTP %d: %s"),
 					Resp->GetResponseCode(), *Resp->GetContentAsString().Left(200));
-				OnComplete.ExecuteIfBound(Response);
+				OnComplete->operator()(MoveTemp(Response));
 				return;
 			}
 
@@ -87,7 +170,7 @@ void FRemoteEmbeddingProvider::GenerateCaptionAsync(
 			if (!FJsonSerializer::Deserialize(Reader, Json) || !Json.IsValid())
 			{
 				Response.ErrorMessage = TEXT("Failed to parse JSON response");
-				OnComplete.ExecuteIfBound(Response);
+				OnComplete->operator()(MoveTemp(Response));
 				return;
 			}
 
@@ -95,14 +178,14 @@ void FRemoteEmbeddingProvider::GenerateCaptionAsync(
 			if (Json->TryGetStringField(TEXT("error"), ServerError))
 			{
 				Response.ErrorMessage = FString::Printf(TEXT("Server error: %s"), *ServerError);
-				OnComplete.ExecuteIfBound(Response);
+				OnComplete->operator()(MoveTemp(Response));
 				return;
 			}
 
 			if (!Json->TryGetStringField(TEXT("caption"), Response.Caption))
 			{
 				Response.ErrorMessage = TEXT("Missing 'caption' in response");
-				OnComplete.ExecuteIfBound(Response);
+				OnComplete->operator()(MoveTemp(Response));
 				return;
 			}
 
@@ -119,18 +202,28 @@ void FRemoteEmbeddingProvider::GenerateCaptionAsync(
 				}
 			}
 
-			OnComplete.ExecuteIfBound(Response);
+			OnComplete->operator()(MoveTemp(Response));
 		});
 
-	Http->ProcessRequest();
+
+	if (RequestOwner)
+	{
+		//  use ref counting to delete itself
+		new FSemanticHttpRequest(Http, *RequestOwner);
+	}
+	else
+	{
+		Http->ProcessRequest();
+	}
 }
 
 void FRemoteEmbeddingProvider::GenerateEmbeddingAsync(
-	const FString& Text,
-	FOnEmbeddingComplete OnComplete)
+	const FStringView Text,
+	FOnEmbeddingComplete&& InOnComplete,
+	DerivedData::FRequestOwner* RequestOwner)
 {
 	TSharedRef<FJsonObject> Payload = MakeShared<FJsonObject>();
-	Payload->SetStringField(TEXT("text"), Text);
+	Payload->SetStringField(TEXT("text"), FString(Text));
 
 	FString PayloadStr;
 	TSharedRef<TJsonWriter<>> Writer = TJsonWriterFactory<>::Create(&PayloadStr);
@@ -144,14 +237,14 @@ void FRemoteEmbeddingProvider::GenerateEmbeddingAsync(
 	Http->SetTimeout(30.0f);
 
 	Http->OnProcessRequestComplete().BindLambda(
-		[this, OnComplete](FHttpRequestPtr /*Req*/, FHttpResponsePtr Resp, bool bConnected)
+		[this, OnComplete = MakeShared<FOnEmbeddingComplete>(MoveTemp(InOnComplete))](FHttpRequestPtr /*Req*/, FHttpResponsePtr Resp, bool bConnected)
 		{
 			FEmbeddingResponse Response;
 
 			if (!bConnected || !Resp.IsValid())
 			{
 				Response.ErrorMessage = TEXT("Connection failed");
-				OnComplete.ExecuteIfBound(Response);
+				OnComplete->operator()(MoveTemp(Response));
 				return;
 			}
 
@@ -159,7 +252,7 @@ void FRemoteEmbeddingProvider::GenerateEmbeddingAsync(
 			{
 				Response.ErrorMessage = FString::Printf(TEXT("HTTP %d: %s"),
 					Resp->GetResponseCode(), *Resp->GetContentAsString().Left(200));
-				OnComplete.ExecuteIfBound(Response);
+				OnComplete->operator()(MoveTemp(Response));
 				return;
 			}
 
@@ -168,7 +261,7 @@ void FRemoteEmbeddingProvider::GenerateEmbeddingAsync(
 			if (!FJsonSerializer::Deserialize(Reader, Json) || !Json.IsValid())
 			{
 				Response.ErrorMessage = TEXT("Failed to parse JSON response");
-				OnComplete.ExecuteIfBound(Response);
+				OnComplete->operator()(MoveTemp(Response));
 				return;
 			}
 
@@ -176,7 +269,7 @@ void FRemoteEmbeddingProvider::GenerateEmbeddingAsync(
 			if (Json->TryGetStringField(TEXT("error"), ServerError))
 			{
 				Response.ErrorMessage = FString::Printf(TEXT("Server error: %s"), *ServerError);
-				OnComplete.ExecuteIfBound(Response);
+				OnComplete->operator()(MoveTemp(Response));
 				return;
 			}
 
@@ -184,7 +277,7 @@ void FRemoteEmbeddingProvider::GenerateEmbeddingAsync(
 			if (!Json->TryGetArrayField(TEXT("embedding"), EmbArray))
 			{
 				Response.ErrorMessage = TEXT("Missing 'embedding' in response");
-				OnComplete.ExecuteIfBound(Response);
+				OnComplete->operator()(MoveTemp(Response));
 				return;
 			}
 
@@ -200,7 +293,7 @@ void FRemoteEmbeddingProvider::GenerateEmbeddingAsync(
 				{
 					Response.Embedding.Empty();
 					Response.ErrorMessage = TEXT("Invalid value in embedding array");
-					OnComplete.ExecuteIfBound(Response);
+					OnComplete->operator()(MoveTemp(Response));
 					return;
 				}
 			}
@@ -214,10 +307,18 @@ void FRemoteEmbeddingProvider::GenerateEmbeddingAsync(
 
 			Json->TryGetStringField(TEXT("model_version"), Response.ModelVersion);
 
-			OnComplete.ExecuteIfBound(Response);
+			OnComplete->operator()(MoveTemp(Response));
 		});
 
-	Http->ProcessRequest();
+	if (RequestOwner)
+	{
+		//  use ref counting to delete itself
+		new FSemanticHttpRequest(Http, *RequestOwner);
+	}
+	else
+	{
+		Http->ProcessRequest();
+	}
 }
 
 TArray<uint8> FRemoteEmbeddingProvider::BuildMultipartBody(
@@ -287,3 +388,4 @@ TArray<uint8> FRemoteEmbeddingProvider::BuildMultipartBody(
 	Append(End);
 	return Body;
 }
+}
diff --git a/Engine/Plugins/Experimental/SemanticSearch/Source/SemanticSearch/Private/RemoteEmbeddingProvider.h b/Engine/Plugins/Experimental/SemanticSearch/Source/SemanticSearch/Private/RemoteEmbeddingProvider.h
--- a/Engine/Plugins/Experimental/SemanticSearch/Source/SemanticSearch/Private/RemoteEmbeddingProvider.h
+++ b/Engine/Plugins/Experimental/SemanticSearch/Source/SemanticSearch/Private/RemoteEmbeddingProvider.h
@@ -7,6 +7,8 @@
 #include "Interfaces/IHttpRequest.h"
 #include "Interfaces/IHttpResponse.h"
 
+namespace UE::SemanticSearch::Private
+{
 /**
  *  Remote embedding provider that communicates with Python backend.
  *
@@ -19,8 +21,8 @@ public:
 	explicit FRemoteEmbeddingProvider(const FString& InBaseURL = FString());
 
 	virtual int32 GetEmbeddingDimension() const override;
-	virtual void GenerateCaptionAsync(const FCaptionRequest& Request, FOnCaptionComplete OnComplete) override;
-	virtual void GenerateEmbeddingAsync(const FString& Text, FOnEmbeddingComplete OnComplete) override;
+	virtual void GenerateCaptionAsync(FCaptionRequest&& Request, FOnCaptionComplete&& OnComplete, DerivedData::FRequestOwner* RequestOwner) override;
+	virtual void GenerateEmbeddingAsync(const FStringView Text, FOnEmbeddingComplete&& OnComplete, DerivedData::FRequestOwner* RequestOwner) override;
 
 	FString GetBaseURL() const { return BaseURL; }
 
@@ -31,3 +33,4 @@ private:
 	FString BaseURL;
 	TAtomic<int32> EmbeddingDimension;
 };
+}
diff --git a/Engine/Plugins/Experimental/SemanticSearch/Source/SemanticSearch/Private/SemanticSearchDerivedDataBuilder.cpp b/Engine/Plugins/Experimental/SemanticSearch/Source/SemanticSearch/Private/SemanticSearchDerivedDataBuilder.cpp
new file mode 100644
--- /dev/null
+++ b/Engine/Plugins/Experimental/SemanticSearch/Source/SemanticSearch/Private/SemanticSearchDerivedDataBuilder.cpp
@@ -0,0 +1,341 @@
+// Copyright Epic Games, Inc. All Rights Reserved.
+
+#include "SemanticSearchDerivedDataBuilder.h"
+
+#include "Containers/SharedString.h"
+#include "DerivedDataCache.h"
+#include "Interfaces/IAssetProcessor.h"
+#include "Misc/Guid.h"
+#include "Misc/ScopeRWLock.h"
+#include "SemanticSearchModule.h"
+#include "Serialization/MemoryHasher.h"
+#include "Serialization/MemoryReader.h"
+#include "Serialization/MemoryWriter.h"
+#include "String/Find.h"
+
+namespace UE::SemanticSearch::Private
+{
+
+DerivedData::FValueId GetEmbeddingDatachunkId()
+{
+	return DerivedData::FValueId::FromName(TEXT("EmbeddingData"));
+}
+
+DerivedData::FValueId GetCaptionDatachunkId()
+{
+	return DerivedData::FValueId::FromName(TEXT("CaptionData"));
+}
+
+void SerializeEmbedding(FEmbeddingDerivedData& Embedding, FArchive& Archive)
+{
+	Archive << Embedding.Embedding;
+}
+
+void SerializeCaption(FCaptionDerivedData& Caption, FArchive& Archive)
+{
+	Archive << Caption.Caption;
+	Archive << Caption.Keywords;
+}
+
+DerivedData::ECachePolicy GetCachePolicy()
+{
+	using namespace DerivedData;
+	// We will need to revisit the non deterministic part of that if we are to cook the embedding data in a cook editor at some point.
+	return ECachePolicy::Default | ECachePolicy::NonDeterministic | ECachePolicy::KeepAlive;
+}
+
+FSemanticSearchBuildCacheTask::FSemanticSearchBuildCacheTask(const FAssetData& InAssetData, TSharedRef<IAssetProcessor>&& InAssetProcessor)
+	: DerivedData::FRequestOwner(DerivedData::EPriority::Low)
+	, Asset(MakeShared<FAssetData>(InAssetData))
+	, AssetProcessor(MoveTemp(InAssetProcessor))
+	, AssetKey(GenerateCacheKey())
+{
+}
+
+FSemanticSearchBuildCacheTask::FSemanticSearchBuildCacheTask(FAssetData&& InAssetData, TSharedRef<IAssetProcessor>&& InAssetProcessor)
+	: DerivedData::FRequestOwner(DerivedData::EPriority::Low)
+	, Asset(MakeShared<FAssetData>(MoveTemp(InAssetData)))
+	, AssetProcessor(MoveTemp(InAssetProcessor))
+	, AssetKey(GenerateCacheKey())
+{
+}
+
+void FSemanticSearchBuildCacheTask::GetEmbeddingData(TUniqueFunction<void(FEmbeddingDerivedData&& EmbeddingData, FString&& ErrorMessage)>&& InOnDataAvailable)
+{
+	{
+		FReadScopeLock ReadLock(IsBuildingDerivedDataLock);
+		if (bIsBuildingDerivedData)
+		{
+			QueuedEmbeddingRequests.Enqueue(MoveTemp(InOnDataAvailable));
+			return;
+		}
+	}
+
+	using namespace DerivedData;
+
+	FCacheGetChunkRequest Request;
+	Request.Key = AssetKey;
+	Request.Id = GetEmbeddingDatachunkId();
+	Request.Policy = GetCachePolicy();
+
+	GetCache().GetChunks({ Request }, *this, [OnDataAvailable = MoveTemp(InOnDataAvailable), this](FCacheGetChunkResponse&& Response) mutable
+		{
+			if (Response.Status == EStatus::Ok && Response.RawData.GetSize() > 0)
+			{
+				FEmbeddingDerivedData Data;
+				
+				FMemoryReaderView Reader(Response.RawData.GetView(), true);
+				SerializeEmbedding(Data, Reader);
+				Response.RawData.Reset();
+				
+				OnDataAvailable(MoveTemp(Data), FString());
+			}
+			else if (Response.Status != EStatus::Canceled)
+			{
+				QueuedEmbeddingRequests.Enqueue(MoveTemp(OnDataAvailable));
+				this->BuildDerivedData();
+			}
+			else
+			{
+				// Something canceled the request
+				OnDataAvailable(FEmbeddingDerivedData(), FString(TEXT("The data request was cancelled")));
+			}
+		});
+}
+
+void FSemanticSearchBuildCacheTask::GetCaptionData(TUniqueFunction<void(FCaptionDerivedData&& CaptionData, FString&& ErrorMessage)>&& InOnDataAvailable)
+{
+	{
+		FReadScopeLock ReadLock(IsBuildingDerivedDataLock);
+		if (bIsBuildingDerivedData)
+		{
+			QueuedCaptionRequests.Enqueue(MoveTemp(InOnDataAvailable));
+			return;
+		}
+	}
+
+	using namespace DerivedData;
+	FCacheGetChunkRequest Request;
+	Request.Key = AssetKey;
+	Request.Id = GetCaptionDatachunkId();
+	Request.Policy = GetCachePolicy();
+
+	GetCache().GetChunks({ Request }, *this, [OnDataAvailable = MoveTemp(InOnDataAvailable), this](FCacheGetChunkResponse&& Response) mutable
+		{
+			if (Response.Status == EStatus::Ok && Response.RawData.GetSize() > 0)
+			{
+				FCaptionDerivedData Data;
+				FMemoryReaderView Reader(Response.RawData.GetView(), true);
+				SerializeCaption(Data, Reader);
+				Response.RawData.Reset();
+
+				OnDataAvailable(MoveTemp(Data), FString());
+			}
+			else if (Response.Status != EStatus::Canceled)
+			{
+				QueuedCaptionRequests.Enqueue(MoveTemp(OnDataAvailable));
+				BuildDerivedData();
+			}
+			else
+			{
+				// Something canceled the request
+				OnDataAvailable(FCaptionDerivedData(), FString(TEXT("The data request was cancelled")));
+			}
+		});
+}
+
+DerivedData::FCacheKey FSemanticSearchBuildCacheTask::GenerateCacheKey() const
+{
+	static FGuid SemanticSearchDerivedDataGuid(TEXT("A3D4A6CE-B6C0-4E41-ABC5-B4C2999D0095"));
+	FMemoryHasherBlake3 HashWriter;
+	HashWriter << SemanticSearchDerivedDataGuid;
+	AssetProcessor->GenerateAssetHash(*Asset, HashWriter);
+
+	// Todo we need to make the embedding provider part of the key
+
+	DerivedData::FCacheKey CacheKey;
+	CacheKey.Hash = HashWriter.Finalize();
+	CacheKey.Bucket = DerivedData::FCacheBucket(TEXT("SemanticSearch"));
+	return CacheKey;
+}
+
+void FSemanticSearchBuildCacheTask::BuildDerivedData()
+{
+	{
+		FWriteScopeLock WriteLock(IsBuildingDerivedDataLock);
+		bIsBuildingDerivedData = true;
+	}
+
+	AssetProcessor->GenerateCaptionRequest(Asset, false, *this, [this](bool bHasGeneratedRequest, FCaptionRequest&& CaptionRequest, FString&& ErrorMessage)
+		{
+			if (bHasGeneratedRequest)
+			{
+				if (TSharedPtr<IEmbeddingProvider> EmbeddingProvider = FSemanticSearchModule::Get().GetRegisteredEmbeddingProvider())
+				{
+					EmbeddingProvider->GenerateCaptionAsync(MoveTemp(CaptionRequest), [this, EmbeddingProvider](FCaptionResponse&& Caption)
+						{
+							if (Caption.ErrorMessage.IsEmpty())
+							{
+								EmbeddingProvider->GenerateEmbeddingAsync(Caption.Caption, [this, UsedCaption = MoveTemp(Caption)](FEmbeddingResponse&& Embedding) mutable
+									{
+										if (Embedding.ErrorMessage.IsEmpty())
+										{
+											FCaptionDerivedData CaptionDerivedData;
+											CaptionDerivedData.Caption = MoveTemp(UsedCaption.Caption);
+											CaptionDerivedData.Keywords = MoveTemp(UsedCaption.Keywords);
+											FEmbeddingDerivedData EmbeddingDerivedData;
+											EmbeddingDerivedData.Embedding = Embedding.Embedding;
+
+											if (AssetKey == GenerateCacheKey())
+											{
+												using namespace DerivedData;
+
+												FCacheRecordBuilder RecordBuilder(AssetKey);
+
+												// Save the caption and keywords
+												{
+													TArray<uint8> Buffer;
+													FMemoryWriter Writer(Buffer);
+													SerializeCaption(CaptionDerivedData, Writer);
+													RecordBuilder.AddValue(GetCaptionDatachunkId(), MakeSharedBufferFromArray(MoveTemp(Buffer)));
+												}
+
+												// Save the embeddings
+												{
+													TArray<uint8> Buffer;
+													FMemoryWriter Writer(Buffer);
+													SerializeEmbedding(EmbeddingDerivedData, Writer);
+													RecordBuilder.AddValue(GetEmbeddingDatachunkId(), MakeSharedBufferFromArray(MoveTemp(Buffer)));
+												}
+
+												DerivedData::FCachePutRequest PutRequest{ FSharedString(FString::Format(TEXT("%s SemanticSearchData"), {FStringFormatArg(Asset->GetObjectPathString())})), RecordBuilder.Build(), GetCachePolicy() | ECachePolicy::SkipData };
+
+
+												DerivedData::GetCache().Put(
+													{ PutRequest },
+													*this,
+													[this](DerivedData::FCachePutResponse&& Response)
+													{
+														if (Response.Status == EStatus::Ok)
+														{
+															UE_LOGF(LogSemanticSearch, Log, "Successfully stored semantic data in DDC for: %ls", *Asset->GetObjectPathString());
+														}
+														else
+														{
+															UE_LOGF(LogSemanticSearch, Warning, "Failed to store data in DDC for: %ls (%ls)", *Asset->GetObjectPathString(), *WriteToString<256>(Response.Status));
+														}
+													});
+
+												BroadcastQueuedRequests(MoveTemp(EmbeddingDerivedData), MoveTemp(CaptionDerivedData), FString());
+
+											}
+											else
+											{
+												UE_LOGF(LogSemanticSearch, Log, "Data used to generate the semantic data changed while building it for: %ls, Result won't be saved to DDC since they might be inaccurate", *Asset->GetObjectPathString());
+												BroadcastQueuedRequests(MoveTemp(EmbeddingDerivedData), MoveTemp(CaptionDerivedData), FString());
+											}
+										}
+										else
+										{
+											// Consider passing the caption data?
+											BroadcastQueuedRequests({}, {}, MoveTemp(Embedding.ErrorMessage));
+										}
+
+									}, this);
+							}
+							else
+							{
+								BroadcastQueuedRequests({}, {}, MoveTemp(Caption.ErrorMessage));
+							}
+						}, this);
+				}
+				else
+				{
+					BroadcastQueuedRequests({}, {}, FString(TEXT("No Embedding provider is active")));
+				}
+			}
+			else
+			{
+				BroadcastQueuedRequests({}, {}, MoveTemp(ErrorMessage));
+			}
+		});
+}
+
+
+void FSemanticSearchBuildCacheTask::BroadcastQueuedRequests(FEmbeddingDerivedData&& EmbeddingData, FCaptionDerivedData&& CaptionData, FString&& ErrorMessage)
+{
+	FWriteScopeLock WriteLock(IsBuildingDerivedDataLock);
+	bIsBuildingDerivedData = false;
+
+	FString CopiedErrorMessage;
+	FString* ErrorMessageToUse = &ErrorMessage;
+
+	// Broadcast Embedding Requests
+	{
+		FEmbeddingDerivedData CopiedEmbeddingData;
+		FEmbeddingDerivedData* EmbeddingDataToUse = &EmbeddingData;
+
+		TUniqueFunction<void(FEmbeddingDerivedData&&, FString&&)> Callback;
+		while (QueuedEmbeddingRequests.Dequeue(Callback))
+		{
+			if (Callback)
+			{
+				bool bSwapCurrentToCopied = false;
+				if (!QueuedEmbeddingRequests.IsEmpty())
+				{
+					bSwapCurrentToCopied = true;
+					CopiedEmbeddingData = *EmbeddingDataToUse;
+					CopiedErrorMessage = *ErrorMessageToUse;
+				}
+
+				Callback(MoveTemp(*EmbeddingDataToUse), MoveTemp(*ErrorMessageToUse));
+
+				if (bSwapCurrentToCopied)
+				{
+					EmbeddingDataToUse = &CopiedEmbeddingData;
+					ErrorMessageToUse = &CopiedErrorMessage;
+				}
+				else
+				{
+					break;
+				}
+			}
+		}
+	}
+
+
+	// Broadcast Caption Request
+	{
+		FCaptionDerivedData CopiedCaptionData;
+		FCaptionDerivedData* CaptionDataToUse = &CaptionData;
+
+		TUniqueFunction<void(FCaptionDerivedData&&, FString&&)> Callback;
+		while (QueuedCaptionRequests.Dequeue(Callback))
+		{
+			if (Callback)
+			{
+				bool bSwapCurrentToCopied = false;
+				if (!QueuedCaptionRequests.IsEmpty())
+				{
+					bSwapCurrentToCopied = true;
+					CopiedCaptionData = *CaptionDataToUse;
+					CopiedErrorMessage = *ErrorMessageToUse;
+				}
+
+				Callback(MoveTemp(*CaptionDataToUse), MoveTemp(*ErrorMessageToUse));
+
+				if (bSwapCurrentToCopied)
+				{
+					CaptionDataToUse = &CopiedCaptionData;
+					ErrorMessageToUse = &CopiedErrorMessage;
+				}
+				else
+				{
+					break;
+				}
+			}
+		}
+	}
+}
+
+}
diff --git a/Engine/Plugins/Experimental/SemanticSearch/Source/SemanticSearch/Private/SemanticSearchDerivedDataBuilder.h b/Engine/Plugins/Experimental/SemanticSearch/Source/SemanticSearch/Private/SemanticSearchDerivedDataBuilder.h
new file mode 100644
--- /dev/null
+++ b/Engine/Plugins/Experimental/SemanticSearch/Source/SemanticSearch/Private/SemanticSearchDerivedDataBuilder.h
@@ -0,0 +1,62 @@
+// Copyright Epic Games, Inc. All Rights Reserved.
+
+#pragma once
+
+#include "AssetRegistry/AssetData.h"
+#include "Containers/Array.h"
+#include "Containers/Queue.h"
+#include "Containers/UnrealString.h"
+#include "DerivedDataCacheKey.h"
+#include "DerivedDataRequestOwner.h"
+#include "HAL/CriticalSection.h"
+#include "Templates/SharedPointer.h"
+
+
+namespace UE::SemanticSearch
+{
+
+class IAssetProcessor;
+
+namespace Private
+{
+/**
+ * We use a guid in the cpp file to track the version of the data. Make sure to change the guid if those are modified (this will invalidate all the cached data)
+ */
+
+struct FEmbeddingDerivedData
+{
+	TArray<float> Embedding;
+};
+
+struct FCaptionDerivedData
+{
+	FString Caption;
+	TArray<FString> Keywords;
+};
+
+class FSemanticSearchBuildCacheTask : public DerivedData::FRequestOwner
+{
+	FSemanticSearchBuildCacheTask(const FAssetData& InAssetData, TSharedRef<IAssetProcessor>&& InAssetProcessor);
+	FSemanticSearchBuildCacheTask(FAssetData&& InAssetData, TSharedRef<IAssetProcessor>&& InAssetProcessor);
+
+	void GetEmbeddingData(TUniqueFunction<void (FEmbeddingDerivedData&& EmbeddingData, FString&& ErrorMessage)>&& InOnDataAvailable);
+	void GetCaptionData(TUniqueFunction<void (FCaptionDerivedData&& CaptionData, FString&& ErrorMessage)>&& InOnDataAvailable);
+
+private:
+	DerivedData::FCacheKey GenerateCacheKey() const;
+
+	void BuildDerivedData();
+	void BroadcastQueuedRequests(FEmbeddingDerivedData&& EmbeddingData, FCaptionDerivedData&& CaptionData, FString&& ErrorMessage);
+
+	TSharedRef<const FAssetData> Asset;
+	TSharedRef<IAssetProcessor> AssetProcessor;
+	DerivedData::FCacheKey AssetKey;
+
+	// Acs has a lock for the IsBuildingDerivedData flag but also avoid race condition with the queued Requests
+	FRWLock IsBuildingDerivedDataLock;
+	bool bIsBuildingDerivedData = false;
+	TQueue<TUniqueFunction<void (FEmbeddingDerivedData&&, FString&&)>, EQueueMode::Mpsc> QueuedEmbeddingRequests;
+	TQueue<TUniqueFunction<void (FCaptionDerivedData&&, FString&&)>, EQueueMode::Mpsc> QueuedCaptionRequests;
+};
+}
+}
diff --git a/Engine/Plugins/Experimental/SemanticSearch/Source/SemanticSearch/Private/SemanticSearchModule.cpp b/Engine/Plugins/Experimental/SemanticSearch/Source/SemanticSearch/Private/SemanticSearchModule.cpp
--- a/Engine/Plugins/Experimental/SemanticSearch/Source/SemanticSearch/Private/SemanticSearchModule.cpp
+++ b/Engine/Plugins/Experimental/SemanticSearch/Source/SemanticSearch/Private/SemanticSearchModule.cpp
@@ -1,17 +1,42 @@
 // Copyright Epic Games, Inc. All Rights Reserved.
 
 #include "SemanticSearchModule.h"
+
 #include "Modules/ModuleManager.h"
+#include "RemoteEmbeddingProvider.h"
 
 DEFINE_LOG_CATEGORY(LogSemanticSearch);
 
+namespace UE::SemanticSearch::Private
+{
 IMPLEMENT_MODULE(FSemanticSearchModule, SemanticSearch)
+
+FSemanticSearchModule& FSemanticSearchModule::Get()
+{
+	static const FName ModuleName(TEXT("SemanticSearch"));
+	return FModuleManager::Get().GetModuleChecked<FSemanticSearchModule>(ModuleName);
+}
+
 void FSemanticSearchModule::StartupModule()
 {
 	UE_LOG(LogSemanticSearch, Log, TEXT("SemanticSearch module starting up"));
+	RegisteredEmbeddingProvider = MakeShared<FRemoteEmbeddingProvider>();
 }
 
 void FSemanticSearchModule::ShutdownModule()
 {
 	UE_LOG(LogSemanticSearch, Log, TEXT("SemanticSearch module shutting down"));
 }
+
+TSharedPtr<IEmbeddingProvider> FSemanticSearchModule::GetRegisteredEmbeddingProvider()
+{
+	return RegisteredEmbeddingProvider;
+}
+
+void FSemanticSearchModule::RegisterEmbeddingProvider(TUniquePtr<IEmbeddingProvider>&& EmbeddingProvider)
+{
+	RegisteredEmbeddingProvider = MakeShareable<IEmbeddingProvider>(EmbeddingProvider.Release());
+}
+
+}
+
diff --git a/Engine/Plugins/Experimental/SemanticSearch/Source/SemanticSearch/Private/SemanticSearchModule.h b/Engine/Plugins/Experimental/SemanticSearch/Source/SemanticSearch/Private/SemanticSearchModule.h
new file mode 100644
--- /dev/null
+++ b/Engine/Plugins/Experimental/SemanticSearch/Source/SemanticSearch/Private/SemanticSearchModule.h
@@ -0,0 +1,31 @@
+// Copyright Epic Games, Inc. All Rights Reserved.
+
+#pragma once
+
+#include "ISemanticSearchModule.h"
+
+#include "CoreMinimal.h"
+#include "Modules/ModuleInterface.h"
+#include "Templates/SharedPointer.h"
+
+DECLARE_LOG_CATEGORY_EXTERN(LogSemanticSearch, Log, All);
+
+namespace UE::SemanticSearch::Private
+{
+
+class FSemanticSearchModule : public ISemanticSearchModule
+{
+public:
+	static FSemanticSearchModule& Get();
+	virtual TSharedPtr<IEmbeddingProvider> GetRegisteredEmbeddingProvider();
+
+	virtual void StartupModule() override;
+	virtual void ShutdownModule() override;
+
+
+	virtual void RegisterEmbeddingProvider(TUniquePtr<IEmbeddingProvider>&& EmbeddingProvider) override;
+
+private:
+	TSharedPtr<IEmbeddingProvider> RegisteredEmbeddingProvider;
+};
+}
\ No newline at end of file
diff --git a/Engine/Plugins/Experimental/SemanticSearch/Source/SemanticSearch/Private/SemanticSearchTestCommands.cpp b/Engine/Plugins/Experimental/SemanticSearch/Source/SemanticSearch/Private/SemanticSearchTestCommands.cpp
--- a/Engine/Plugins/Experimental/SemanticSearch/Source/SemanticSearch/Private/SemanticSearchTestCommands.cpp
+++ b/Engine/Plugins/Experimental/SemanticSearch/Source/SemanticSearch/Private/SemanticSearchTestCommands.cpp
@@ -6,6 +6,9 @@
 #include "Misc/FileHelper.h"
 #include "Misc/Paths.h"
 
+namespace UE::SemanticSearch::Private
+{
+
 static TSharedPtr<FRemoteEmbeddingProvider> GTestProvider;
 
 static TSharedPtr<FRemoteEmbeddingProvider> GetOrCreateProvider()
@@ -58,8 +61,8 @@ static FAutoConsoleCommand TestCaptionCommand(
 		UE_LOG(LogSemanticSearch, Log, TEXT("Sending caption request (%s, %d bytes)..."), *MimeType, Request.AssetMedia[0].Data.Num());
 
 		TSharedPtr<FRemoteEmbeddingProvider> Provider = GetOrCreateProvider();
-		Provider->GenerateCaptionAsync(Request,
-			FOnCaptionComplete::CreateLambda([](const FCaptionResponse& Response)
+		Provider->GenerateCaptionAsync(MoveTemp(Request),
+			[](const FCaptionResponse& Response)
 			{
 				if (!Response.ErrorMessage.IsEmpty())
 				{
@@ -72,7 +75,7 @@ static FAutoConsoleCommand TestCaptionCommand(
 				{
 					UE_LOG(LogSemanticSearch, Log, TEXT("Keywords: %s"), *FString::Join(Response.Keywords, TEXT(", ")));
 				}
-			}));
+			}, nullptr);
 	})
 );
 
@@ -96,7 +99,7 @@ static FAutoConsoleCommand TestEmbeddingCommand(
 
 		TSharedPtr<FRemoteEmbeddingProvider> Provider = GetOrCreateProvider();
 		Provider->GenerateEmbeddingAsync(Text,
-			FOnEmbeddingComplete::CreateLambda([](const FEmbeddingResponse& Response)
+			[](const FEmbeddingResponse& Response)
 			{
 				if (!Response.ErrorMessage.IsEmpty())
 				{
@@ -121,6 +124,8 @@ static FAutoConsoleCommand TestEmbeddingCommand(
 					}
 					UE_LOG(LogSemanticSearch, Log, TEXT("First %d values: [%s, ...]"), Preview, *Vals);
 				}
-			}));
+			}, nullptr);
 	})
 );
+
+}
diff --git a/Engine/Plugins/Experimental/SemanticSearch/Source/SemanticSearch/Public/ISemanticSearchModule.h b/Engine/Plugins/Experimental/SemanticSearch/Source/SemanticSearch/Public/ISemanticSearchModule.h
new file mode 100644
--- /dev/null
+++ b/Engine/Plugins/Experimental/SemanticSearch/Source/SemanticSearch/Public/ISemanticSearchModule.h
@@ -0,0 +1,24 @@
+// Copyright Epic Games, Inc. All Rights Reserved.
+
+#pragma once
+
+#include "Interfaces/IEmbeddingProvider.h"
+#include "Modules/ModuleInterface.h"
+#include "Modules/ModuleManager.h"
+#include "Templates/SharedPointer.h"
+#include "Templates/UniquePtr.h"
+
+namespace UE::SemanticSearch
+{
+	class ISemanticSearchModule : public IModuleInterface
+	{
+	public:
+		static ISemanticSearchModule& Get()
+		{
+			static const FName ModuleName = "SemanticSearch";
+			return FModuleManager::LoadModuleChecked<ISemanticSearchModule>(ModuleName);
+		}
+
+		virtual void RegisterEmbeddingProvider(TUniquePtr<IEmbeddingProvider>&& EmbeddingProvider) = 0;
+	};
+}
\ No newline at end of file
diff --git a/Engine/Plugins/Experimental/SemanticSearch/Source/SemanticSearch/Public/Interfaces/IAssetProcessor.h b/Engine/Plugins/Experimental/SemanticSearch/Source/SemanticSearch/Public/Interfaces/IAssetProcessor.h
new file mode 100644
--- /dev/null
+++ b/Engine/Plugins/Experimental/SemanticSearch/Source/SemanticSearch/Public/Interfaces/IAssetProcessor.h
@@ -0,0 +1,39 @@
+// Copyright Epic Games, Inc. All Rights Reserved.
+
+#pragma once
+
+#include "Containers/StringFwd.h"
+#include "Interfaces/IEmbeddingProvider.h"
+#include "Serialization/MemoryHasher.h"
+#include "Templates/Function.h"
+
+class UClass;
+struct FAssetData;
+
+
+namespace UE::DerivedData
+{
+	class FRequestOwner;
+}
+
+namespace UE::SemanticSearch
+{
+struct FCaptionRequest;
+
+using FOnRequestComplete = TUniqueFunction<void(bool HasGeneratedCationRequest, FCaptionRequest&& CaptionRequest, FString&& ErrorMessage)>;
+
+class IAssetProcessor
+{
+public:
+	virtual UClass& GetSupportedClass() const = 0;
+	virtual bool SupportDerivedClasses() const = 0;
+
+	virtual void GenerateAssetHash(const FAssetData& Asset, FMemoryHasherBlake3& InHasher) const = 0;
+	virtual FStringView GetAssetProcessorName() const = 0;
+
+	virtual bool GenerateCaptionRequest(const TSharedRef<const FAssetData>& InAsset, bool bCanDisruptUser, DerivedData::FRequestOwner& RequestOwner, FOnRequestComplete OnRequestComplete) const = 0;
+
+	virtual ~IAssetProcessor() = default;
+};
+
+}
\ No newline at end of file
diff --git a/Engine/Plugins/Experimental/SemanticSearch/Source/SemanticSearch/Public/Interfaces/IEmbeddingProvider.h b/Engine/Plugins/Experimental/SemanticSearch/Source/SemanticSearch/Public/Interfaces/IEmbeddingProvider.h
--- a/Engine/Plugins/Experimental/SemanticSearch/Source/SemanticSearch/Public/Interfaces/IEmbeddingProvider.h
+++ b/Engine/Plugins/Experimental/SemanticSearch/Source/SemanticSearch/Public/Interfaces/IEmbeddingProvider.h
@@ -4,106 +4,111 @@
 
 #include "CoreMinimal.h"
 #include "Dom/JsonObject.h"
+#include "Templates/Function.h"
 
-/**
- * A single media payload attached to a caption request.
- */
-struct SEMANTICSEARCH_API FAssetMedia
+namespace UE::DerivedData
 {
-	/** MIME type: "image/png", "application/json", "audio/mpeg", etc. */
-	FString MimeType;
-
-	/** Raw media bytes */
-	TArray<uint8> Data;
-};
-
-/**
- * Request to generate a caption (and keywords) for an asset.
- * Can include media content and structured metadata.
- */
-struct SEMANTICSEARCH_API FCaptionRequest
-{
-	FString AssetPath;
-	FString AssetType;
+	class FRequestOwner;
+}
 
-	/** Asset content — can be combination of images, text, JSON */
-	TArray<FAssetMedia> AssetMedia;
+namespace UE::SemanticSearch
+{
+	/**
+	 * A single media payload attached to a caption request.
+	 */
+	struct FAssetMedia
+	{
+		/** MIME type: "image/png", "application/json", "audio/mpeg", etc. */
+		FString MimeType;
 
-	/** Structured metadata (asset dimension, properties, etc.) */
-	TSharedPtr<FJsonObject> Metadata;
+		/** Raw media bytes */
+		TArray<uint8> Data;
+	};
 
-	/** Content hash for DDC cache key generation */
-	FString ContentHash;
-};
+	/**
+	 * Request to generate a caption (and keywords) for an asset.
+	 * Can include media content and structured metadata.
+	 */
+	struct FCaptionRequest
+	{
+		FString AssetPath;
+		FString AssetType;
 
-/**
- * Response from caption generation.
- * Success if ErrorMessage is empty.
- */
-struct SEMANTICSEARCH_API FCaptionResponse
-{
-	FString ErrorMessage;
+		/** Asset content — can be combination of images, text, JSON */
+		TArray<FAssetMedia> AssetMedia;
 
-	FString Caption;
+		/** Structured metadata (asset dimension, properties, etc.) */
+		TSharedPtr<FJsonObject> Metadata;
+	};
 
-	TArray<FString> Keywords;
-};
+	/**
+	 * Response from caption generation.
+	 * Success if ErrorMessage is empty.
+	 */
+	struct FCaptionResponse
+	{
+		FString ErrorMessage;
 
-/**
- * Response from embedding generation.
- * Success if ErrorMessage is empty.
- */
-struct SEMANTICSEARCH_API FEmbeddingResponse
-{
-	FString ErrorMessage;
+		FString Caption;
 
-	TArray<float> Embedding;
+		TArray<FString> Keywords;
+	};
 
-	/** Version of the model that generated this embedding */
-	FString ModelVersion;
-};
+	/**
+	 * Response from embedding generation.
+	 * Success if ErrorMessage is empty.
+	 */
+	struct FEmbeddingResponse
+	{
+		FString ErrorMessage;
 
-/**
- * Delegate for async caption completion
- */
-DECLARE_DELEGATE_OneParam(FOnCaptionComplete, const FCaptionResponse&);
+		TArray<float> Embedding;
 
-/**
- * Delegate for async embedding completion
- */
-DECLARE_DELEGATE_OneParam(FOnEmbeddingComplete, const FEmbeddingResponse&);
+		/** Version of the model that generated this embedding */
+		FString ModelVersion;
+	};
 
-/**
- * Interface for embedding generation providers.
- * Can support both remote (HTTP API) and local (on-device) implementations.
- */
-class SEMANTICSEARCH_API IEmbeddingProvider
-{
-public:
-	virtual ~IEmbeddingProvider() = default;
 
-	/** Embedding vector dimensionality */
-	virtual int32 GetEmbeddingDimension() const = 0;
+	using FOnCaptionComplete = TUniqueFunction<void (FCaptionResponse&&)>;
 
-	/**
-	 * Generate a caption and keywords for an asset.
-	 * Called during indexing to describe the asset content.
-	 *
-	 * @param Request The caption request with asset media and metadata
-	 * @param OnComplete Callback invoked when caption is ready
-	 */
-	virtual void GenerateCaptionAsync(
-		const FCaptionRequest& Request,
-		FOnCaptionComplete OnComplete) = 0;
+	using FOnEmbeddingComplete = TUniqueFunction<void (FEmbeddingResponse&&)>;
 
 	/**
-	 * Generate embedding for a text string.
-	 * Called during indexing (with caption text) and at search time (with query text).
-	 *
-	 * @param Text The text to embed (caption or search query)
-	 * @param OnComplete Callback invoked when embedding is ready
+	 * Interface for embedding generation providers.
+	 * Can support both remote (HTTP API) and local (on-device) implementations.
 	 */
-	virtual void GenerateEmbeddingAsync(
-		const FString& Text,
-		FOnEmbeddingComplete OnComplete) = 0;
-};
+	class IEmbeddingProvider
+	{
+	public:
+		virtual ~IEmbeddingProvider() = default;
+
+		/** Embedding vector dimensionality */
+		virtual int32 GetEmbeddingDimension() const = 0;
+
+		/**
+		 * Generate a caption and keywords for an asset.
+		 * Called during indexing to describe the asset content.
+		 *
+		 * @param Request The caption request with asset media and metadata
+		 * @param OnComplete Callback invoked when caption is ready
+		 * @param RequestOwner The request owner is optional argument to help the derived data build track it subtasks
+		 */
+		virtual void GenerateCaptionAsync(
+			FCaptionRequest&& Request,
+			FOnCaptionComplete&& OnComplete,
+			DerivedData::FRequestOwner* RequestOwner) = 0;
+
+		/**
+		 * Generate embedding for a text string.
+		 * Called during indexing (with caption text) and at search time (with query text).
+		 *
+		 * @param Text The text to embed (caption or search query)
+		 * @param OnComplete Callback invoked when embedding is ready
+		 * @param RequestOwner The request owner is optional argument to help the derived data build track it subtasks
+		 */
+		virtual void GenerateEmbeddingAsync(
+			const FStringView Text,
+			FOnEmbeddingComplete&& OnComplete,
+			DerivedData::FRequestOwner* RequestOwner) = 0;
+	};
+}
diff --git a/Engine/Plugins/Experimental/SemanticSearch/Source/SemanticSearch/Public/SemanticSearchModule.h b/Engine/Plugins/Experimental/SemanticSearch/Source/SemanticSearch/Public/SemanticSearchModule.h
deleted file mode 100644
diff --git a/Engine/Plugins/Experimental/SemanticSearch/Source/SemanticSearch/SemanticSearch.Build.cs b/Engine/Plugins/Experimental/SemanticSearch/Source/SemanticSearch/SemanticSearch.Build.cs
--- a/Engine/Plugins/Experimental/SemanticSearch/Source/SemanticSearch/SemanticSearch.Build.cs
+++ b/Engine/Plugins/Experimental/SemanticSearch/Source/SemanticSearch/SemanticSearch.Build.cs
@@ -18,9 +18,10 @@ public class SemanticSearch : ModuleRules
 				"AssetRegistry",
 				"Core",
 				"CoreUObject",
+				"DerivedDataCache",
 				"HTTP",
 				"Json",
-			}
-			);
+			});
+
 	}
 }
\ No newline at end of file
