Semantic Search: Add the texture asset processor.

Adds FTextureProcessor, a new IAssetProcessor implementation for UTexture assets (including derived classes)
The processor use the asset name, asset path and the package hash to generate key. It can only generate a key for on disk asset at the moment. The caption request is generated from the thumbnail and a few key meta data that might help interpreting what the asset is.

Added a function to FThumbnailCompressionInterface that allow us to get a mime type for the compressed version of the thumbnail.

#Synced-CL 51543477

diff --git a/Engine/Plugins/Experimental/SemanticSearch/Source/SemanticSearch/Private/AssetProcessors/TextureAssetProcessor.cpp b/Engine/Plugins/Experimental/SemanticSearch/Source/SemanticSearch/Private/AssetProcessors/TextureAssetProcessor.cpp
new file mode 100644
--- /dev/null
+++ b/Engine/Plugins/Experimental/SemanticSearch/Source/SemanticSearch/Private/AssetProcessors/TextureAssetProcessor.cpp
@@ -0,0 +1,186 @@
+// Copyright Epic Games, Inc. All Rights Reserved.
+
+#include "AssetProcessors/TextureAssetProcessor.h"
+
+#include "AssetRegistry/AssetData.h"
+#include "AssetRegistry/IAssetRegistry.h"
+#include "Engine/Texture.h"
+#include "Engine/TextureDefines.h"
+#include "ImageCore.h"
+#include "Misc/ObjectThumbnail.inl"
+#include "Misc/PackageName.h"
+#include "ObjectTools.h"
+#include "Serialization/MemoryHasher.h"
+#include "Templates/UnrealTemplate.h"
+#include "UObject/NameTypes.h"
+#include "UObject/ReflectedTypeAccessors.h"
+
+namespace UE::SemanticSearch::Private
+{
+UClass& FTextureProcessor::GetSupportedClass() const
+{
+	return *UTexture::StaticClass();
+}
+
+bool FTextureProcessor::SupportDerivedClasses() const
+{
+	return true;
+}
+
+bool FTextureProcessor::GenerateAssetHash(const FAssetData& Asset, FMemoryHasherBlake3& InHasher) const
+{
+	static FGuid TextureProcessorDerivedDataGuid(TEXT("FF41B729-F7E6-408D-A164-55BF969100D8"));
+
+	if (Asset.HasAnyPackageFlags(PKG_Cooked))
+	{
+		return false;
+	}
+
+	FAssetPackageData AssetPackageData;
+	if (IAssetRegistry::GetChecked().TryGetAssetPackageData(Asset.PackageName, AssetPackageData) != AssetRegistry::EExists::Exists)
+	{
+		return false;
+	}
+
+	InHasher << TextureProcessorDerivedDataGuid;
+
+	FName PackageName(Asset.PackageName);
+	InHasher << PackageName;
+
+	FName AssetName(Asset.AssetName);
+	InHasher << AssetName;
+
+	// Just use the save hash for now
+	FIoHash PackageHash = AssetPackageData.GetPackageSavedHash();
+	InHasher << PackageHash;
+
+	return true;
+}
+
+void FTextureProcessor::GenerateCaptionRequest(const TSharedRef<const FAssetData>& InAsset, bool bCanDisruptUser, DerivedData::FRequestOwner& RequestOwner, FOnRequestComplete OnRequestComplete) const
+{
+	FString PackagePath;
+	if (FPackageName::DoesPackageExist(InAsset->PackageName.ToString(), &PackagePath))
+	{
+		TSet<FName> ObjectsFullName;
+		FName ObjectFullName = *InAsset->GetFullName();
+		ObjectsFullName.Add(ObjectFullName);
+
+		FThumbnailMap ThumbnailMap;
+		// Don't use the cache to avoid pollution. Consider revisiting this?
+		if (ThumbnailTools::LoadThumbnailsFromPackageDirectly(PackagePath, ObjectsFullName, ThumbnailMap, PackagePath))
+		{
+			if (FObjectThumbnail* Thumbnail = ThumbnailMap.Find(ObjectFullName))
+			{
+				FCaptionRequest CaptionRequest;
+
+				// Convert the format from BGRA8 to RGBA8 inplace;
+				Thumbnail->DecompressImageData();
+				if (!ensure(Thumbnail->GetImage().Format == ERawImageFormat::BGRA8))
+				{
+					OnRequestComplete(false, FCaptionRequest(), FString(TEXT("Thumbnail was using an unexpected data format.")));
+					return;
+				}
+				TArrayView<uint8> UncompressedImage = Thumbnail->AccessImageData();
+				for (int32 PixelIndex = 0; PixelIndex < UncompressedImage.Num(); PixelIndex += 4)
+				{
+					constexpr int32 RedOffsetInBGRA8 = 2;
+					Swap(UncompressedImage[PixelIndex], UncompressedImage[PixelIndex + RedOffsetInBGRA8]);
+				}
+				Thumbnail->CompressImageData();
+				Thumbnail->AccessImageData().Empty();
+
+
+				FAssetMedia Media;
+				Media.MimeType = Thumbnail->GetCompressor()->GetMimeType();
+				Media.Data = MoveTemp(Thumbnail->AccessCompressedImageData());
+				CaptionRequest.AssetMedia.Add(MoveTemp(Media));
+
+				CaptionRequest.AssetPath = InAsset->GetObjectPathString();
+				CaptionRequest.AssetType = InAsset->AssetClassPath.ToString();
+
+				// Meta data
+				CaptionRequest.Metadata = MakeShared<FJsonObject>();
+				FString Value;
+
+				auto SetMetadata = [&CaptionRequest, &InAsset, &Value](FName AssetDataTag, FStringView CaptionMetadataKey)
+				{
+					if (InAsset->GetTagValue(AssetDataTag, Value))
+					{
+						CaptionRequest.Metadata->SetStringField(FString(CaptionMetadataKey), MoveTemp(Value));
+					}
+				};
+
+				auto SetMetadataWithDisplayString = [this, &CaptionRequest, &InAsset, &Value](FName AssetDataTag, FStringView CaptionMetadataKey)
+				{
+					if (InAsset->GetTagValue(AssetDataTag, Value))
+					{
+						const FString* DisplayValue = MetaDataValueStringToDisplayString.Find(Value);
+						CaptionRequest.Metadata->SetStringField(FString(CaptionMetadataKey), DisplayValue ? *DisplayValue : MoveTemp(Value));
+					}
+				};
+
+				static const FName DimensionsTag(TEXT("Dimensions"));
+				SetMetadata(DimensionsTag, TEXTVIEW("Dimensions"));
+
+				static const FName MaxTextureSizeTag(TEXT("MaxTextureSize"));
+				if (InAsset->GetTagValue(MaxTextureSizeTag, Value) && Value != TEXTVIEW("0"))
+				{
+					CaptionRequest.Metadata->SetStringField(FString(TEXTVIEW("MaxTextureSize")), MoveTemp(Value));
+				}
+
+				static const FName WrappingX(TEXT("AddressX"));
+				SetMetadataWithDisplayString(WrappingX, TEXTVIEW("X-axis Tiling Method"));
+
+				static const FName WrappingY(TEXT("AddressY"));
+				SetMetadataWithDisplayString(WrappingY, TEXTVIEW("Y-axis Tiling Method"));
+
+				// Texture arrays
+				static const FName WrappingZ(TEXT("AddressZ"));
+				SetMetadataWithDisplayString(WrappingZ, TEXTVIEW("Z-axis Tiling Method"));
+
+				// Volume texture
+				static const FName WrappingMode(TEXT("AddressMode"));
+				SetMetadataWithDisplayString(WrappingMode, TEXTVIEW("Tiling Method"));
+
+				static const FName CompressionSettings(TEXT("CompressionSettings"));
+				SetMetadataWithDisplayString(CompressionSettings, TEXTVIEW("Compression Settings"));
+
+				static const FName SRGBTag(TEXT("SRGB"));
+				SetMetadata(SRGBTag, TEXTVIEW("sRGB"));
+
+				static const FName VirtualTextureStreamingTag(TEXT("VirtualTextureStreaming"));
+				SetMetadata(VirtualTextureStreamingTag, TEXTVIEW("Virtual Texture Streaming"));
+
+
+				OnRequestComplete(true, MoveTemp(CaptionRequest), FString());
+				return;
+			}
+		}
+	}
+
+	OnRequestComplete(false, FCaptionRequest(), FString(TEXT("Couldn't fetch an on-disk thumbnail for the asset")));
+}
+
+FTextureProcessor::FTextureProcessor()
+{
+	// Reflection data is not thread safe in editor
+	check(IsInGameThread());
+
+	auto PopulateFromEnum = [this](const UEnum* Enum)
+	{
+		for (int32 Index = 0; Index < Enum->NumEnums() - 1; ++Index)
+		{
+			FString Key = Enum->GetNameStringByIndex(Index);
+
+			// We don't want the user localization in our metadata
+			FString DisplayValue = Enum->GetDisplayNameTextByIndex(Index).BuildSourceString();
+			MetaDataValueStringToDisplayString.Add(MoveTemp(Key), MoveTemp(DisplayValue));
+		}
+	};
+
+	PopulateFromEnum(StaticEnum<TextureAddress>());
+	PopulateFromEnum(StaticEnum<TextureCompressionSettings>());
+}
+
+}
diff --git a/Engine/Plugins/Experimental/SemanticSearch/Source/SemanticSearch/Private/AssetProcessors/TextureAssetProcessor.h b/Engine/Plugins/Experimental/SemanticSearch/Source/SemanticSearch/Private/AssetProcessors/TextureAssetProcessor.h
new file mode 100644
--- /dev/null
+++ b/Engine/Plugins/Experimental/SemanticSearch/Source/SemanticSearch/Private/AssetProcessors/TextureAssetProcessor.h
@@ -0,0 +1,27 @@
+// Copyright Epic Games, Inc. All Rights Reserved.
+
+#pragma once
+
+#include "Interfaces/IAssetProcessor.h"
+
+namespace UE::SemanticSearch::Private
+{
+class FTextureProcessor : public IAssetProcessor
+{
+public:
+	// Must be created on the game thread
+	FTextureProcessor();
+
+	// Thread safe
+	virtual UClass& GetSupportedClass()  const override;
+	virtual bool SupportDerivedClasses() const override;
+	virtual bool GenerateAssetHash(const FAssetData& Asset, FMemoryHasherBlake3& InHasher) const override;
+	virtual void GenerateCaptionRequest(const TSharedRef<const FAssetData>& InAsset, bool bCanDisruptUser, DerivedData::FRequestOwner& RequestOwner, FOnRequestComplete OnRequestComplete) const override;
+
+private:
+	// Cache the display value as those are not thread safe
+	TMap<FString, FString> MetaDataValueStringToDisplayString;
+
+	};
+}
+
diff --git a/Engine/Plugins/Experimental/SemanticSearch/Source/SemanticSearch/Private/SemanticSearchDerivedDataBuilder.cpp b/Engine/Plugins/Experimental/SemanticSearch/Source/SemanticSearch/Private/SemanticSearchDerivedDataBuilder.cpp
--- a/Engine/Plugins/Experimental/SemanticSearch/Source/SemanticSearch/Private/SemanticSearchDerivedDataBuilder.cpp
+++ b/Engine/Plugins/Experimental/SemanticSearch/Source/SemanticSearch/Private/SemanticSearchDerivedDataBuilder.cpp
@@ -41,21 +41,21 @@ DerivedData::ECachePolicy GetCachePolicy()
 {
 	using namespace DerivedData;
 	// We will need to revisit the non deterministic part of that if we are to cook the embedding data in a cook editor at some point.
-	return ECachePolicy::Default | ECachePolicy::NonDeterministic | ECachePolicy::KeepAlive;
+	return ECachePolicy::Default | ECachePolicy::NonDeterministic;
 }
 
-FSemanticSearchBuildCacheTask::FSemanticSearchBuildCacheTask(const FAssetData& InAssetData, TSharedRef<IAssetProcessor>&& InAssetProcessor)
+FSemanticSearchBuildCacheTask::FSemanticSearchBuildCacheTask(const FAssetData& InAssetData, const TSharedRef<IAssetProcessor>& InAssetProcessor)
 	: DerivedData::FRequestOwner(DerivedData::EPriority::Low)
 	, Asset(MakeShared<FAssetData>(InAssetData))
-	, AssetProcessor(MoveTemp(InAssetProcessor))
+	, AssetProcessor(InAssetProcessor)
 	, AssetKey(GenerateCacheKey())
 {
 }
 
-FSemanticSearchBuildCacheTask::FSemanticSearchBuildCacheTask(FAssetData&& InAssetData, TSharedRef<IAssetProcessor>&& InAssetProcessor)
+FSemanticSearchBuildCacheTask::FSemanticSearchBuildCacheTask(FAssetData&& InAssetData, const TSharedRef<IAssetProcessor>& InAssetProcessor)
 	: DerivedData::FRequestOwner(DerivedData::EPriority::Low)
 	, Asset(MakeShared<FAssetData>(MoveTemp(InAssetData)))
-	, AssetProcessor(MoveTemp(InAssetProcessor))
+	, AssetProcessor(InAssetProcessor)
 	, AssetKey(GenerateCacheKey())
 {
 }
@@ -149,7 +149,11 @@ DerivedData::FCacheKey FSemanticSearchBuildCacheTask::GenerateCacheKey() const
 	static FGuid SemanticSearchDerivedDataGuid(TEXT("A3D4A6CE-B6C0-4E41-ABC5-B4C2999D0095"));
 	FMemoryHasherBlake3 HashWriter;
 	HashWriter << SemanticSearchDerivedDataGuid;
-	AssetProcessor->GenerateAssetHash(*Asset, HashWriter);
+	if (!AssetProcessor->GenerateAssetHash(*Asset, HashWriter))
+	{
+		// Invalid Key Generation
+		return DerivedData::FCacheKey{};
+	}
 
 	// Todo we need to make the embedding provider part of the key
 
@@ -176,7 +180,8 @@ void FSemanticSearchBuildCacheTask::BuildDerivedData()
 						{
 							if (Caption.ErrorMessage.IsEmpty())
 							{
-								EmbeddingProvider->GenerateEmbeddingAsync(Caption.Caption, [this, UsedCaption = MoveTemp(Caption)](FEmbeddingResponse&& Embedding) mutable
+								FString CaptionString = Caption.Caption;
+								EmbeddingProvider->GenerateEmbeddingAsync(CaptionString, [this, UsedCaption = MoveTemp(Caption)](FEmbeddingResponse&& Embedding) mutable
 									{
 										if (Embedding.ErrorMessage.IsEmpty())
 										{
@@ -214,15 +219,15 @@ void FSemanticSearchBuildCacheTask::BuildDerivedData()
 												DerivedData::GetCache().Put(
 													{ PutRequest },
 													*this,
-													[this](DerivedData::FCachePutResponse&& Response)
+													[AssetPtr = Asset](DerivedData::FCachePutResponse&& Response)
 													{
 														if (Response.Status == EStatus::Ok)
 														{
-															UE_LOGF(LogSemanticSearch, Log, "Successfully stored semantic data in DDC for: %ls", *Asset->GetObjectPathString());
+															UE_LOGF(LogSemanticSearch, Log, "Successfully stored semantic data in DDC for: %ls", *AssetPtr->GetObjectPathString());
 														}
 														else
 														{
-															UE_LOGF(LogSemanticSearch, Warning, "Failed to store data in DDC for: %ls (%ls)", *Asset->GetObjectPathString(), *WriteToString<256>(Response.Status));
+															UE_LOGF(LogSemanticSearch, Warning, "Failed to store data in DDC for: %ls (%ls)", *AssetPtr->GetObjectPathString(), *WriteToString<256>(Response.Status));
 														}
 													});
 
@@ -264,8 +269,33 @@ void FSemanticSearchBuildCacheTask::BuildDerivedData()
 
 void FSemanticSearchBuildCacheTask::BroadcastQueuedRequests(FEmbeddingDerivedData&& EmbeddingData, FCaptionDerivedData&& CaptionData, FString&& ErrorMessage)
 {
-	FWriteScopeLock WriteLock(IsBuildingDerivedDataLock);
-	bIsBuildingDerivedData = false;
+	TArray<TUniqueFunction<void(FEmbeddingDerivedData&&, FString&&)>, TInlineAllocator<4>> EmbeddingRequests;
+	TArray<TUniqueFunction<void(FCaptionDerivedData&&, FString&&)>, TInlineAllocator<4>> CaptionRequests;
+
+	{
+		FWriteScopeLock WriteLock(IsBuildingDerivedDataLock);
+		bIsBuildingDerivedData = false;
+
+		// Grab Embedding Requests
+		{
+			TUniqueFunction<void(FEmbeddingDerivedData&&, FString&&)> Callback;
+			while (QueuedEmbeddingRequests.Dequeue(Callback))
+			{
+				EmbeddingRequests.Add(MoveTemp(Callback));
+			}
+		}
+
+
+		// Grab Caption Request
+		{
+			TUniqueFunction<void(FCaptionDerivedData&&, FString&&)> Callback;
+			while (QueuedCaptionRequests.Dequeue(Callback))
+			{
+				CaptionRequests.Add(MoveTemp(Callback));
+			}
+		}
+	}
+
 
 	FString CopiedErrorMessage;
 	FString* ErrorMessageToUse = &ErrorMessage;
@@ -274,14 +304,15 @@ void FSemanticSearchBuildCacheTask::BroadcastQueuedRequests(FEmbeddingDerivedDat
 	{
 		FEmbeddingDerivedData CopiedEmbeddingData;
 		FEmbeddingDerivedData* EmbeddingDataToUse = &EmbeddingData;
+		int32 Count = 0;
 
-		TUniqueFunction<void(FEmbeddingDerivedData&&, FString&&)> Callback;
-		while (QueuedEmbeddingRequests.Dequeue(Callback))
+		for (TUniqueFunction<void(FEmbeddingDerivedData&&, FString&&)>& Callback : EmbeddingRequests)
 		{
+			++Count;
 			if (Callback)
 			{
 				bool bSwapCurrentToCopied = false;
-				if (!QueuedEmbeddingRequests.IsEmpty())
+				if (Count < EmbeddingRequests.Num())
 				{
 					bSwapCurrentToCopied = true;
 					CopiedEmbeddingData = *EmbeddingDataToUse;
@@ -303,19 +334,19 @@ void FSemanticSearchBuildCacheTask::BroadcastQueuedRequests(FEmbeddingDerivedDat
 		}
 	}
 
-
 	// Broadcast Caption Request
 	{
 		FCaptionDerivedData CopiedCaptionData;
 		FCaptionDerivedData* CaptionDataToUse = &CaptionData;
+		int32 Count = 0;
 
-		TUniqueFunction<void(FCaptionDerivedData&&, FString&&)> Callback;
-		while (QueuedCaptionRequests.Dequeue(Callback))
+		for (TUniqueFunction<void(FCaptionDerivedData&&, FString&&)>& Callback : CaptionRequests)
 		{
+			++Count;
 			if (Callback)
 			{
 				bool bSwapCurrentToCopied = false;
-				if (!QueuedCaptionRequests.IsEmpty())
+				if (Count < CaptionRequests.Num())
 				{
 					bSwapCurrentToCopied = true;
 					CopiedCaptionData = *CaptionDataToUse;
diff --git a/Engine/Plugins/Experimental/SemanticSearch/Source/SemanticSearch/Private/SemanticSearchDerivedDataBuilder.h b/Engine/Plugins/Experimental/SemanticSearch/Source/SemanticSearch/Private/SemanticSearchDerivedDataBuilder.h
--- a/Engine/Plugins/Experimental/SemanticSearch/Source/SemanticSearch/Private/SemanticSearchDerivedDataBuilder.h
+++ b/Engine/Plugins/Experimental/SemanticSearch/Source/SemanticSearch/Private/SemanticSearchDerivedDataBuilder.h
@@ -36,8 +36,9 @@ struct FCaptionDerivedData
 
 class FSemanticSearchBuildCacheTask : public DerivedData::FRequestOwner
 {
-	FSemanticSearchBuildCacheTask(const FAssetData& InAssetData, TSharedRef<IAssetProcessor>&& InAssetProcessor);
-	FSemanticSearchBuildCacheTask(FAssetData&& InAssetData, TSharedRef<IAssetProcessor>&& InAssetProcessor);
+public:
+	FSemanticSearchBuildCacheTask(const FAssetData& InAssetData, const TSharedRef<IAssetProcessor>& InAssetProcessor);
+	FSemanticSearchBuildCacheTask(FAssetData&& InAssetData, const TSharedRef<IAssetProcessor>& InAssetProcessor);
 
 	void GetEmbeddingData(TUniqueFunction<void (FEmbeddingDerivedData&& EmbeddingData, FString&& ErrorMessage)>&& InOnDataAvailable);
 	void GetCaptionData(TUniqueFunction<void (FCaptionDerivedData&& CaptionData, FString&& ErrorMessage)>&& InOnDataAvailable);
diff --git a/Engine/Plugins/Experimental/SemanticSearch/Source/SemanticSearch/Public/Interfaces/IAssetProcessor.h b/Engine/Plugins/Experimental/SemanticSearch/Source/SemanticSearch/Public/Interfaces/IAssetProcessor.h
--- a/Engine/Plugins/Experimental/SemanticSearch/Source/SemanticSearch/Public/Interfaces/IAssetProcessor.h
+++ b/Engine/Plugins/Experimental/SemanticSearch/Source/SemanticSearch/Public/Interfaces/IAssetProcessor.h
@@ -20,7 +20,7 @@ namespace UE::SemanticSearch
 {
 struct FCaptionRequest;
 
-using FOnRequestComplete = TUniqueFunction<void(bool HasGeneratedCationRequest, FCaptionRequest&& CaptionRequest, FString&& ErrorMessage)>;
+using FOnRequestComplete = TUniqueFunction<void(bool bHasGeneratedCaptionRequest, FCaptionRequest&& CaptionRequest, FString&& ErrorMessage)>;
 
 class IAssetProcessor
 {
@@ -28,10 +28,8 @@ public:
 	virtual UClass& GetSupportedClass() const = 0;
 	virtual bool SupportDerivedClasses() const = 0;
 
-	virtual void GenerateAssetHash(const FAssetData& Asset, FMemoryHasherBlake3& InHasher) const = 0;
-	virtual FStringView GetAssetProcessorName() const = 0;
-
-	virtual bool GenerateCaptionRequest(const TSharedRef<const FAssetData>& InAsset, bool bCanDisruptUser, DerivedData::FRequestOwner& RequestOwner, FOnRequestComplete OnRequestComplete) const = 0;
+	virtual bool GenerateAssetHash(const FAssetData& Asset, FMemoryHasherBlake3& InHasher) const = 0;
+	virtual void GenerateCaptionRequest(const TSharedRef<const FAssetData>& InAsset, bool bCanDisruptUser, DerivedData::FRequestOwner& RequestOwner, FOnRequestComplete OnRequestComplete) const = 0;
 
 	virtual ~IAssetProcessor() = default;
 };
diff --git a/Engine/Plugins/Experimental/SemanticSearch/Source/SemanticSearch/SemanticSearch.Build.cs b/Engine/Plugins/Experimental/SemanticSearch/Source/SemanticSearch/SemanticSearch.Build.cs
--- a/Engine/Plugins/Experimental/SemanticSearch/Source/SemanticSearch/SemanticSearch.Build.cs
+++ b/Engine/Plugins/Experimental/SemanticSearch/Source/SemanticSearch/SemanticSearch.Build.cs
@@ -19,8 +19,11 @@ public class SemanticSearch : ModuleRules
 				"Core",
 				"CoreUObject",
 				"DerivedDataCache",
+				"Engine",
 				"HTTP",
+				"ImageCore",
 				"Json",
+				"UnrealEd",
 			});
 
 	}
diff --git a/Engine/Source/Runtime/Core/Public/Misc/ObjectThumbnail.h b/Engine/Source/Runtime/Core/Public/Misc/ObjectThumbnail.h
--- a/Engine/Source/Runtime/Core/Public/Misc/ObjectThumbnail.h
+++ b/Engine/Source/Runtime/Core/Public/Misc/ObjectThumbnail.h
@@ -53,6 +53,11 @@ public:
 	 * @return true if compression is lossless
 	 */
 	virtual bool IsLosslessCompression() const = 0;
+
+	/**
+	 * Get the mime type of the underlying compressed data
+	 */
+	virtual FStringView GetMimeType() const = 0;
 };
 
 
diff --git a/Engine/Source/Runtime/Engine/Private/UnrealEngine.cpp b/Engine/Source/Runtime/Engine/Private/UnrealEngine.cpp
--- a/Engine/Source/Runtime/Engine/Private/UnrealEngine.cpp
+++ b/Engine/Source/Runtime/Engine/Private/UnrealEngine.cpp
@@ -1457,6 +1457,11 @@ public:
 	{
 		return true;
 	}
+
+	FStringView GetMimeType() const override
+	{
+		return FStringView(TEXT("image/png"));
+	}
 };
 
 
@@ -1556,6 +1561,11 @@ public:
 	{
 		return false;
 	}
+
+	FStringView GetMimeType() const override
+	{
+		return FStringView(TEXT("image/jpeg"));
+	}
 };
 
 
