Refactor the texture asset processor to use a base implementation that will be shared with other asset types.
This is the Thumbnail Base Pocessor.

Added an AssetProcessorUtils for the utils logic shared  accross the processors.
Added the Static Mesh, Skeletal Mesh, Material and  Material Instance processor.

#Synced-CL 51761360

diff --git a/Engine/Plugins/Experimental/SemanticSearch/Source/SemanticSearch/Private/AssetProcessors/AssetProcessorUtils.cpp b/Engine/Plugins/Experimental/SemanticSearch/Source/SemanticSearch/Private/AssetProcessors/AssetProcessorUtils.cpp
new file mode 100644
--- /dev/null
+++ b/Engine/Plugins/Experimental/SemanticSearch/Source/SemanticSearch/Private/AssetProcessors/AssetProcessorUtils.cpp
@@ -0,0 +1,55 @@
+// Copyright Epic Games, Inc. All Rights Reserved.
+
+#include "AssetProcessors/AssetProcessorUtils.h"
+
+#include "AssetRegistry/AssetData.h"
+#include "Dom/JsonObject.h"
+#include "UObject/Class.h"
+
+namespace UE::SemanticSearch::Private
+{
+
+void SetMetadata(
+	const TSharedPtr<FJsonObject>& Metadata,
+	const TSharedRef<const FAssetData>& InAsset,
+	FName AssetDataTag,
+	FStringView CaptionMetadataKey)
+{
+	FString Value;
+	if (InAsset->GetTagValue(AssetDataTag, Value))
+	{
+		Metadata->SetStringField(FString(CaptionMetadataKey), MoveTemp(Value));
+	}
+}
+
+void SetMetadataWithDisplayString(
+	const TSharedPtr<FJsonObject>& Metadata,
+	const TSharedRef<const FAssetData>& InAsset,
+	FName AssetDataTag,
+	FStringView CaptionMetadataKey,
+	const TMap<FString, FString>& DisplayMap)
+{
+	FString Value;
+	if (InAsset->GetTagValue(AssetDataTag, Value))
+	{
+		const FString* DisplayValue = DisplayMap.Find(Value);
+		Metadata->SetStringField(FString(CaptionMetadataKey), DisplayValue ? *DisplayValue : MoveTemp(Value));
+	}
+}
+
+void PopulateFromEnum(TMap<FString, FString>& OutMap, const UEnum* Enum)
+{
+	// Reflection data is not thread safe in editor
+	check(IsInGameThread());
+
+	for (int32 Index = 0; Index < Enum->NumEnums() - 1; ++Index)
+	{
+		FString Key = Enum->GetNameStringByIndex(Index);
+
+		// We don't want the user localization in our metadata
+		FString DisplayValue = Enum->GetDisplayNameTextByIndex(Index).BuildSourceString();
+		OutMap.Add(MoveTemp(Key), MoveTemp(DisplayValue));
+	}
+}
+
+}
diff --git a/Engine/Plugins/Experimental/SemanticSearch/Source/SemanticSearch/Private/AssetProcessors/AssetProcessorUtils.h b/Engine/Plugins/Experimental/SemanticSearch/Source/SemanticSearch/Private/AssetProcessors/AssetProcessorUtils.h
new file mode 100644
--- /dev/null
+++ b/Engine/Plugins/Experimental/SemanticSearch/Source/SemanticSearch/Private/AssetProcessors/AssetProcessorUtils.h
@@ -0,0 +1,54 @@
+// Copyright Epic Games, Inc. All Rights Reserved.
+
+#pragma once
+
+#include "CoreMinimal.h"
+
+class FJsonObject;
+struct FAssetData;
+class UEnum;
+
+namespace UE::SemanticSearch::Private
+{
+
+/**
+ * Reads AssetDataTag from InAsset and, if present, writes the raw value into Metadata under CaptionMetadataKey.
+ *
+ * @param Metadata              JSON object to write the field into
+ * @param InAsset               Asset whose tag store is queried
+ * @param AssetDataTag          Tag name to look up on the asset
+ * @param CaptionMetadataKey    Key used when writing the field into Metadata
+ */
+void SetMetadata(
+	const TSharedPtr<FJsonObject>& Metadata,
+	const TSharedRef<const FAssetData>& InAsset,
+	FName AssetDataTag,
+	FStringView CaptionMetadataKey);
+
+/**
+ * Reads AssetDataTag from InAsset and, if present, writes the value into Metadata under CaptionMetadataKey.
+ * If DisplayMap contains an entry for the raw value, the mapped display string is written instead.
+ *
+ * @param Metadata              JSON object to write the field into
+ * @param InAsset               Asset whose tag store is queried
+ * @param AssetDataTag          Tag name to look up on the asset
+ * @param CaptionMetadataKey    Key used when writing the field into Metadata
+ * @param DisplayMap            Optional remapping from raw enum name strings to human-readable display strings
+ */
+void SetMetadataWithDisplayString(
+	const TSharedPtr<FJsonObject>& Metadata,
+	const TSharedRef<const FAssetData>& InAsset,
+	FName AssetDataTag,
+	FStringView CaptionMetadataKey,
+	const TMap<FString, FString>& DisplayMap);
+
+/**
+ * Populates OutMap with name-to-display-string entries for every enumerator in Enum.
+ * Must be called on the game thread as reflection data is not thread safe in the editor.
+ *
+ * @param OutMap    Map to populate with enumerator name to display string pairs
+ * @param Enum      Enum whose enumerators are iterated
+ */
+void PopulateFromEnum(TMap<FString, FString>& OutMap, const UEnum* Enum);
+
+}
diff --git a/Engine/Plugins/Experimental/SemanticSearch/Source/SemanticSearch/Private/AssetProcessors/MaterialAssetProcessor.cpp b/Engine/Plugins/Experimental/SemanticSearch/Source/SemanticSearch/Private/AssetProcessors/MaterialAssetProcessor.cpp
new file mode 100644
--- /dev/null
+++ b/Engine/Plugins/Experimental/SemanticSearch/Source/SemanticSearch/Private/AssetProcessors/MaterialAssetProcessor.cpp
@@ -0,0 +1,123 @@
+// Copyright Epic Games, Inc. All Rights Reserved.
+
+#include "AssetProcessors/MaterialAssetProcessor.h"
+
+#include "AssetProcessors/AssetProcessorUtils.h"
+
+#include "AssetRegistry/AssetData.h"
+#include "Materials/Material.h"
+#include "MaterialDomain.h"
+#include "UObject/ReflectedTypeAccessors.h"
+
+namespace UE::SemanticSearch::Private
+{
+
+constexpr FStringView FMaterialProcessor::GetProcessSubBucketName() const
+{
+	return TEXTVIEW("FMaterialProcessor");
+}
+
+UClass& FMaterialProcessor::GetSupportedClass() const
+{
+	return *UMaterial::StaticClass();
+}
+
+bool FMaterialProcessor::SupportDerivedClasses() const
+{
+	return false;
+}
+
+FString FMaterialProcessor::DecodeShadingModelsBitmask(const FString& BitmaskStr) const
+{
+	// Expected asset registry format: "(ShadingModelField=4)"
+	constexpr FStringView Prefix = TEXTVIEW("(ShadingModelField=");
+	constexpr FStringView Suffix = TEXTVIEW(")");
+	if (!ensureMsgf(BitmaskStr.StartsWith(Prefix) && BitmaskStr.EndsWith(Suffix),
+		TEXT("FMaterialProcessor::DecodeShadingModelsBitmask: unexpected format '%s'"), *BitmaskStr))
+	{
+		return FString();
+	}
+
+	const FString NumberStr = BitmaskStr.Mid(Prefix.Len(), BitmaskStr.Len() - Prefix.Len() - 1);
+	const int32 Bitmask = FCString::Atoi(*NumberStr);
+
+	TArray<FString> ActiveModels;
+	for (int32 Index = 0; Index < ShadingModelIndexToName.Num(); ++Index)
+	{
+		if (Bitmask & (1 << Index))
+		{
+			const FString& EnumName = ShadingModelIndexToName[Index];
+			const FString* Display = MetaDataValueStringToDisplayString.Find(EnumName);
+			ActiveModels.Add(Display ? *Display : EnumName);
+		}
+	}
+
+	return FString::Join(ActiveModels, TEXT(", "));
+}
+
+TSharedPtr<FJsonObject> FMaterialProcessor::GetMetadata(const TSharedRef<const FAssetData>& InAsset) const
+{
+	TSharedPtr<FJsonObject> Metadata = MakeShared<FJsonObject>();
+
+	static const FName MaterialDomainTag(TEXT("MaterialDomain"));
+	SetMetadataWithDisplayString(Metadata, InAsset, MaterialDomainTag, TEXTVIEW("Material Domain"), MetaDataValueStringToDisplayString);
+
+	static const FName BlendModeTag(TEXT("BlendMode"));
+	SetMetadataWithDisplayString(Metadata, InAsset, BlendModeTag, TEXTVIEW("Blend Mode"), MetaDataValueStringToDisplayString);
+
+	static const FName ShadingModelTag(TEXT("ShadingModel"));
+	SetMetadataWithDisplayString(Metadata, InAsset, ShadingModelTag, TEXTVIEW("Shading Model"), MetaDataValueStringToDisplayString);
+
+	static const FName ShadingModelsTag(TEXT("ShadingModels"));
+	{
+		FString Value;
+		if (InAsset->GetTagValue(ShadingModelsTag, Value))
+		{
+			FString Decoded = DecodeShadingModelsBitmask(Value);
+			if (!Decoded.IsEmpty())
+			{
+				Metadata->SetStringField(FString(TEXTVIEW("Shading Models")), MoveTemp(Decoded));
+			}
+		}
+	}
+
+	static const FName DecalResponseTag(TEXT("MaterialDecalResponse"));
+	SetMetadataWithDisplayString(Metadata, InAsset, DecalResponseTag, TEXTVIEW("Decal Response"), MetaDataValueStringToDisplayString);
+
+	static const FName TranslucencyLightingModeTag(TEXT("TranslucencyLightingMode"));
+	SetMetadataWithDisplayString(Metadata, InAsset, TranslucencyLightingModeTag, TEXTVIEW("Translucency Lighting Mode"), MetaDataValueStringToDisplayString);
+
+	static const FName HasSceneColorTag(TEXT("HasSceneColor"));
+	SetMetadata(Metadata, InAsset, HasSceneColorTag, TEXTVIEW("Has Scene Color"));
+
+	static const FName HasPerInstanceRandomTag(TEXT("HasPerInstanceRandom"));
+	SetMetadata(Metadata, InAsset, HasPerInstanceRandomTag, TEXTVIEW("Has Per Instance Random"));
+
+	static const FName HasPerInstanceCustomDataTag(TEXT("HasPerInstanceCustomData"));
+	SetMetadata(Metadata, InAsset, HasPerInstanceCustomDataTag, TEXTVIEW("Has Per Instance Custom Data"));
+
+	static const FName HasVertexInterpolatorTag(TEXT("HasVertexInterpolator"));
+	SetMetadata(Metadata, InAsset, HasVertexInterpolatorTag, TEXTVIEW("Has Vertex Interpolator"));
+
+	return Metadata;
+}
+
+FMaterialProcessor::FMaterialProcessor()
+{
+	check(IsInGameThread());
+
+	const UEnum* ShadingModelEnum = StaticEnum<EMaterialShadingModel>();
+	ShadingModelIndexToName.SetNum(MSM_NUM);
+	for (int32 Index = 0; Index < MSM_NUM; ++Index)
+	{
+		ShadingModelIndexToName[Index] = ShadingModelEnum->GetNameStringByIndex(Index);
+	}
+
+	PopulateFromEnum(MetaDataValueStringToDisplayString, StaticEnum<EMaterialDomain>());
+	PopulateFromEnum(MetaDataValueStringToDisplayString, StaticEnum<EBlendMode>());
+	PopulateFromEnum(MetaDataValueStringToDisplayString, StaticEnum<EMaterialShadingModel>());
+	PopulateFromEnum(MetaDataValueStringToDisplayString, StaticEnum<EMaterialDecalResponse>());
+	PopulateFromEnum(MetaDataValueStringToDisplayString, StaticEnum<ETranslucencyLightingMode>());
+}
+
+}
diff --git a/Engine/Plugins/Experimental/SemanticSearch/Source/SemanticSearch/Private/AssetProcessors/MaterialAssetProcessor.h b/Engine/Plugins/Experimental/SemanticSearch/Source/SemanticSearch/Private/AssetProcessors/MaterialAssetProcessor.h
new file mode 100644
--- /dev/null
+++ b/Engine/Plugins/Experimental/SemanticSearch/Source/SemanticSearch/Private/AssetProcessors/MaterialAssetProcessor.h
@@ -0,0 +1,34 @@
+// Copyright Epic Games, Inc. All Rights Reserved.
+
+#pragma once
+
+#include "Implementations/ThumbnailBaseProcessor.h"
+
+namespace UE::SemanticSearch::Private
+{
+
+class FMaterialProcessor : public FThumbnailBaseAssetProcessor
+{
+public:
+	// Must be created on the game thread
+	FMaterialProcessor();
+
+	constexpr virtual FStringView GetProcessSubBucketName() const override;
+
+	// Thread safe
+	virtual UClass& GetSupportedClass() const override;
+	virtual bool SupportDerivedClasses() const override;
+
+	virtual TSharedPtr<FJsonObject> GetMetadata(const TSharedRef<const FAssetData>& InAsset) const override;
+
+private:
+	FString DecodeShadingModelsBitmask(const FString& BitmaskStr) const;
+
+	// Cache the display values as those are not thread safe
+	TMap<FString, FString> MetaDataValueStringToDisplayString;
+
+	// Cache of EMaterialShadingModel index to name string — reflection data is not thread safe
+	TArray<FString> ShadingModelIndexToName;
+};
+
+}
diff --git a/Engine/Plugins/Experimental/SemanticSearch/Source/SemanticSearch/Private/AssetProcessors/MaterialInstanceAssetProcessor.cpp b/Engine/Plugins/Experimental/SemanticSearch/Source/SemanticSearch/Private/AssetProcessors/MaterialInstanceAssetProcessor.cpp
new file mode 100644
--- /dev/null
+++ b/Engine/Plugins/Experimental/SemanticSearch/Source/SemanticSearch/Private/AssetProcessors/MaterialInstanceAssetProcessor.cpp
@@ -0,0 +1,52 @@
+// Copyright Epic Games, Inc. All Rights Reserved.
+
+#include "AssetProcessors/MaterialInstanceAssetProcessor.h"
+
+#include "AssetProcessors/AssetProcessorUtils.h"
+
+#include "AssetRegistry/AssetData.h"
+#include "Materials/MaterialInstance.h"
+
+namespace UE::SemanticSearch::Private
+{
+
+FMaterialInstanceProcessor::FMaterialInstanceProcessor() = default;
+
+constexpr FStringView FMaterialInstanceProcessor::GetProcessSubBucketName() const
+{
+	return TEXTVIEW("FMaterialInstanceProcessor");
+}
+
+UClass& FMaterialInstanceProcessor::GetSupportedClass() const
+{
+	return *UMaterialInstance::StaticClass();
+}
+
+bool FMaterialInstanceProcessor::SupportDerivedClasses() const
+{
+	return false;
+}
+
+TSharedPtr<FJsonObject> FMaterialInstanceProcessor::GetMetadata(const TSharedRef<const FAssetData>& InAsset) const
+{
+	TSharedPtr<FJsonObject> Metadata = MakeShared<FJsonObject>();
+
+	static const FName ParentTag(TEXT("Parent"));
+	SetMetadata(Metadata, InAsset, ParentTag, TEXTVIEW("Parent"));
+
+	static const FName HasSceneColorTag(TEXT("HasSceneColor"));
+	SetMetadata(Metadata, InAsset, HasSceneColorTag, TEXTVIEW("Has Scene Color"));
+
+	static const FName HasPerInstanceRandomTag(TEXT("HasPerInstanceRandom"));
+	SetMetadata(Metadata, InAsset, HasPerInstanceRandomTag, TEXTVIEW("Has Per Instance Random"));
+
+	static const FName HasPerInstanceCustomDataTag(TEXT("HasPerInstanceCustomData"));
+	SetMetadata(Metadata, InAsset, HasPerInstanceCustomDataTag, TEXTVIEW("Has Per Instance Custom Data"));
+
+	static const FName HasVertexInterpolatorTag(TEXT("HasVertexInterpolator"));
+	SetMetadata(Metadata, InAsset, HasVertexInterpolatorTag, TEXTVIEW("Has Vertex Interpolator"));
+
+	return Metadata;
+}
+
+}
diff --git a/Engine/Plugins/Experimental/SemanticSearch/Source/SemanticSearch/Private/AssetProcessors/MaterialInstanceAssetProcessor.h b/Engine/Plugins/Experimental/SemanticSearch/Source/SemanticSearch/Private/AssetProcessors/MaterialInstanceAssetProcessor.h
new file mode 100644
--- /dev/null
+++ b/Engine/Plugins/Experimental/SemanticSearch/Source/SemanticSearch/Private/AssetProcessors/MaterialInstanceAssetProcessor.h
@@ -0,0 +1,24 @@
+// Copyright Epic Games, Inc. All Rights Reserved.
+
+#pragma once
+
+#include "Implementations/ThumbnailBaseProcessor.h"
+
+namespace UE::SemanticSearch::Private
+{
+
+class FMaterialInstanceProcessor : public FThumbnailBaseAssetProcessor
+{
+public:
+	FMaterialInstanceProcessor();
+
+	constexpr virtual FStringView GetProcessSubBucketName() const override;
+
+	// Thread safe
+	virtual UClass& GetSupportedClass() const override;
+	virtual bool SupportDerivedClasses() const override;
+
+	virtual TSharedPtr<FJsonObject> GetMetadata(const TSharedRef<const FAssetData>& InAsset) const override;
+};
+
+}
diff --git a/Engine/Plugins/Experimental/SemanticSearch/Source/SemanticSearch/Private/AssetProcessors/SkeletalMeshAssetProcessor.cpp b/Engine/Plugins/Experimental/SemanticSearch/Source/SemanticSearch/Private/AssetProcessors/SkeletalMeshAssetProcessor.cpp
new file mode 100644
--- /dev/null
+++ b/Engine/Plugins/Experimental/SemanticSearch/Source/SemanticSearch/Private/AssetProcessors/SkeletalMeshAssetProcessor.cpp
@@ -0,0 +1,64 @@
+// Copyright Epic Games, Inc. All Rights Reserved.
+
+#include "AssetProcessors/SkeletalMeshAssetProcessor.h"
+
+#include "AssetProcessors/AssetProcessorUtils.h"
+
+#include "AssetRegistry/AssetData.h"
+#include "Engine/SkeletalMesh.h"
+
+namespace UE::SemanticSearch::Private
+{
+
+FSkeletalMeshProcessor::FSkeletalMeshProcessor() = default;
+
+constexpr FStringView FSkeletalMeshProcessor::GetProcessSubBucketName() const
+{
+	return TEXTVIEW("FSkeletalMeshProcessor");
+}
+
+UClass& FSkeletalMeshProcessor::GetSupportedClass() const
+{
+	return *USkeletalMesh::StaticClass();
+}
+
+bool FSkeletalMeshProcessor::SupportDerivedClasses() const
+{
+	return false;
+}
+
+TSharedPtr<FJsonObject> FSkeletalMeshProcessor::GetMetadata(const TSharedRef<const FAssetData>& InAsset) const
+{
+	TSharedPtr<FJsonObject> Metadata = MakeShared<FJsonObject>();
+
+	static const FName BonesTag(TEXT("Bones"));
+	SetMetadata(Metadata, InAsset, BonesTag, TEXTVIEW("Bones"));
+
+	static const FName MorphTargetsTag(TEXT("MorphTargets"));
+	SetMetadata(Metadata, InAsset, MorphTargetsTag, TEXTVIEW("Morph Targets"));
+
+	static const FName NaniteEnabledTag(TEXT("NaniteEnabled"));
+	SetMetadata(Metadata, InAsset, NaniteEnabledTag, TEXTVIEW("Nanite Enabled"));
+
+	static const FName NaniteTrianglesTag(TEXT("NaniteTriangles"));
+	SetMetadata(Metadata, InAsset, NaniteTrianglesTag, TEXTVIEW("Nanite Triangles"));
+
+	static const FName NaniteVerticesTag(TEXT("NaniteVertices"));
+	SetMetadata(Metadata, InAsset, NaniteVerticesTag, TEXTVIEW("Nanite Vertices"));
+
+	static const FName TrianglesTag(TEXT("Triangles"));
+	SetMetadata(Metadata, InAsset, TrianglesTag, TEXTVIEW("Triangles"));
+
+	static const FName VerticesTag(TEXT("Vertices"));
+	SetMetadata(Metadata, InAsset, VerticesTag, TEXTVIEW("Vertices"));
+
+	static const FName PhysicsAssetTag(TEXT("PhysicsAsset"));
+	SetMetadata(Metadata, InAsset, PhysicsAssetTag, TEXTVIEW("Physics Asset"));
+
+	static const FName SkeletonTag(TEXT("Skeleton"));
+	SetMetadata(Metadata, InAsset, SkeletonTag, TEXTVIEW("Skeleton"));
+
+	return Metadata;
+}
+
+}
diff --git a/Engine/Plugins/Experimental/SemanticSearch/Source/SemanticSearch/Private/AssetProcessors/SkeletalMeshAssetProcessor.h b/Engine/Plugins/Experimental/SemanticSearch/Source/SemanticSearch/Private/AssetProcessors/SkeletalMeshAssetProcessor.h
new file mode 100644
--- /dev/null
+++ b/Engine/Plugins/Experimental/SemanticSearch/Source/SemanticSearch/Private/AssetProcessors/SkeletalMeshAssetProcessor.h
@@ -0,0 +1,24 @@
+// Copyright Epic Games, Inc. All Rights Reserved.
+
+#pragma once
+
+#include "Implementations/ThumbnailBaseProcessor.h"
+
+namespace UE::SemanticSearch::Private
+{
+
+class FSkeletalMeshProcessor : public FThumbnailBaseAssetProcessor
+{
+public:
+	FSkeletalMeshProcessor();
+
+	constexpr virtual FStringView GetProcessSubBucketName() const override;
+
+	// Thread safe
+	virtual UClass& GetSupportedClass() const override;
+	virtual bool SupportDerivedClasses() const override;
+
+	virtual TSharedPtr<FJsonObject> GetMetadata(const TSharedRef<const FAssetData>& InAsset) const override;
+};
+
+}
diff --git a/Engine/Plugins/Experimental/SemanticSearch/Source/SemanticSearch/Private/AssetProcessors/StaticMeshAssetProcessor.cpp b/Engine/Plugins/Experimental/SemanticSearch/Source/SemanticSearch/Private/AssetProcessors/StaticMeshAssetProcessor.cpp
new file mode 100644
--- /dev/null
+++ b/Engine/Plugins/Experimental/SemanticSearch/Source/SemanticSearch/Private/AssetProcessors/StaticMeshAssetProcessor.cpp
@@ -0,0 +1,66 @@
+// Copyright Epic Games, Inc. All Rights Reserved.
+
+#include "AssetProcessors/StaticMeshAssetProcessor.h"
+
+#include "AssetProcessors/AssetProcessorUtils.h"
+
+#include "AssetRegistry/AssetData.h"
+#include "Engine/StaticMesh.h"
+#include "BodySetupEnums.h"
+#include "UObject/ReflectedTypeAccessors.h"
+
+namespace UE::SemanticSearch::Private
+{
+
+constexpr FStringView FStaticMeshProcessor::GetProcessSubBucketName() const
+{
+	return TEXTVIEW("FStaticMeshProcessor");
+}
+
+UClass& FStaticMeshProcessor::GetSupportedClass() const
+{
+	return *UStaticMesh::StaticClass();
+}
+
+bool FStaticMeshProcessor::SupportDerivedClasses() const
+{
+	return false;
+}
+
+TSharedPtr<FJsonObject> FStaticMeshProcessor::GetMetadata(const TSharedRef<const FAssetData>& InAsset) const
+{
+	TSharedPtr<FJsonObject> Metadata = MakeShared<FJsonObject>();
+
+	static const FName CollisionComplexityTag(TEXT("CollisionComplexity"));
+	SetMetadataWithDisplayString(Metadata, InAsset, CollisionComplexityTag, TEXTVIEW("Collision Complexity"), MetaDataValueStringToDisplayString);
+
+	static const FName CollisionPrimsTag(TEXT("CollisionPrims"));
+	SetMetadata(Metadata, InAsset, CollisionPrimsTag, TEXTVIEW("Collision Primitives"));
+
+	static const FName DefaultCollisionTag(TEXT("DefaultCollision"));
+	SetMetadata(Metadata, InAsset, DefaultCollisionTag, TEXTVIEW("Default Collision"));
+
+	static const FName NaniteEnabledTag(TEXT("NaniteEnabled"));
+	SetMetadata(Metadata, InAsset, NaniteEnabledTag, TEXTVIEW("Nanite Enabled"));
+
+	static const FName NaniteTrianglesTag(TEXT("NaniteTriangles"));
+	SetMetadata(Metadata, InAsset, NaniteTrianglesTag, TEXTVIEW("Nanite Triangles"));
+
+	static const FName NaniteVerticesTag(TEXT("NaniteVertices"));
+	SetMetadata(Metadata, InAsset, NaniteVerticesTag, TEXTVIEW("Nanite Vertices"));
+
+	static const FName TrianglesTag(TEXT("Triangles"));
+	SetMetadata(Metadata, InAsset, TrianglesTag, TEXTVIEW("Triangles"));
+
+	static const FName VerticesTag(TEXT("Vertices"));
+	SetMetadata(Metadata, InAsset, VerticesTag, TEXTVIEW("Vertices"));
+
+	return Metadata;
+}
+
+FStaticMeshProcessor::FStaticMeshProcessor()
+{
+	PopulateFromEnum(MetaDataValueStringToDisplayString, StaticEnum<ECollisionTraceFlag>());
+}
+
+}
diff --git a/Engine/Plugins/Experimental/SemanticSearch/Source/SemanticSearch/Private/AssetProcessors/StaticMeshAssetProcessor.h b/Engine/Plugins/Experimental/SemanticSearch/Source/SemanticSearch/Private/AssetProcessors/StaticMeshAssetProcessor.h
new file mode 100644
--- /dev/null
+++ b/Engine/Plugins/Experimental/SemanticSearch/Source/SemanticSearch/Private/AssetProcessors/StaticMeshAssetProcessor.h
@@ -0,0 +1,29 @@
+// Copyright Epic Games, Inc. All Rights Reserved.
+
+#pragma once
+
+#include "Implementations/ThumbnailBaseProcessor.h"
+
+namespace UE::SemanticSearch::Private
+{
+
+class FStaticMeshProcessor : public FThumbnailBaseAssetProcessor
+{
+public:
+	// Must be created on the game thread
+	FStaticMeshProcessor();
+
+	constexpr virtual FStringView GetProcessSubBucketName() const override;
+
+	// Thread safe
+	virtual UClass& GetSupportedClass() const override;
+	virtual bool SupportDerivedClasses() const override;
+
+	virtual TSharedPtr<FJsonObject> GetMetadata(const TSharedRef<const FAssetData>& InAsset) const override;
+
+private:
+	// Cache the display values as those are not thread safe
+	TMap<FString, FString> MetaDataValueStringToDisplayString;
+};
+
+}
diff --git a/Engine/Plugins/Experimental/SemanticSearch/Source/SemanticSearch/Private/AssetProcessors/TextureAssetProcessor.cpp b/Engine/Plugins/Experimental/SemanticSearch/Source/SemanticSearch/Private/AssetProcessors/TextureAssetProcessor.cpp
--- a/Engine/Plugins/Experimental/SemanticSearch/Source/SemanticSearch/Private/AssetProcessors/TextureAssetProcessor.cpp
+++ b/Engine/Plugins/Experimental/SemanticSearch/Source/SemanticSearch/Private/AssetProcessors/TextureAssetProcessor.cpp
@@ -2,21 +2,21 @@
 
 #include "AssetProcessors/TextureAssetProcessor.h"
 
+#include "AssetProcessors/AssetProcessorUtils.h"
+
 #include "AssetRegistry/AssetData.h"
-#include "AssetRegistry/IAssetRegistry.h"
 #include "Engine/Texture.h"
 #include "Engine/TextureDefines.h"
-#include "ImageCore.h"
-#include "Misc/ObjectThumbnail.inl"
-#include "Misc/PackageName.h"
-#include "ObjectTools.h"
-#include "Serialization/MemoryHasher.h"
-#include "Templates/UnrealTemplate.h"
-#include "UObject/NameTypes.h"
 #include "UObject/ReflectedTypeAccessors.h"
 
 namespace UE::SemanticSearch::Private
 {
+
+constexpr FStringView FTextureProcessor::GetProcessSubBucketName() const
+{
+	return TEXTVIEW("FTextureProcessor");
+}
+
 UClass& FTextureProcessor::GetSupportedClass() const
 {
 	return *UTexture::StaticClass();
@@ -27,160 +27,52 @@ bool FTextureProcessor::SupportDerivedClasses() const
 	return true;
 }
 
-bool FTextureProcessor::GenerateAssetHash(const FAssetData& Asset, FMemoryHasherBlake3& InHasher) const
+TSharedPtr<FJsonObject> FTextureProcessor::GetMetadata(const TSharedRef<const FAssetData>& InAsset) const
 {
-	static FGuid TextureProcessorDerivedDataGuid(TEXT("FF41B729-F7E6-408D-A164-55BF969100D8"));
+	TSharedPtr<FJsonObject> Metadata = MakeShared<FJsonObject>();
 
-	if (Asset.HasAnyPackageFlags(PKG_Cooked))
-	{
-		return false;
-	}
+	static const FName DimensionsTag(TEXT("Dimensions"));
+	SetMetadata(Metadata, InAsset, DimensionsTag, TEXTVIEW("Dimensions"));
 
-	FAssetPackageData AssetPackageData;
-	if (IAssetRegistry::GetChecked().TryGetAssetPackageData(Asset.PackageName, AssetPackageData) != AssetRegistry::EExists::Exists)
+	static const FName MaxTextureSizeTag(TEXT("MaxTextureSize"));
 	{
-		return false;
+		FString Value;
+		if (InAsset->GetTagValue(MaxTextureSizeTag, Value) && Value != TEXTVIEW("0"))
+		{
+			Metadata->SetStringField(FString(TEXTVIEW("MaxTextureSize")), MoveTemp(Value));
+		}
 	}
 
-	InHasher << TextureProcessorDerivedDataGuid;
+	static const FName WrappingX(TEXT("AddressX"));
+	SetMetadataWithDisplayString(Metadata, InAsset, WrappingX, TEXTVIEW("X-axis Tiling Method"), MetaDataValueStringToDisplayString);
 
-	FName PackageName(Asset.PackageName);
-	InHasher << PackageName;
+	static const FName WrappingY(TEXT("AddressY"));
+	SetMetadataWithDisplayString(Metadata, InAsset, WrappingY, TEXTVIEW("Y-axis Tiling Method"), MetaDataValueStringToDisplayString);
 
-	FName AssetName(Asset.AssetName);
-	InHasher << AssetName;
+	// Texture arrays
+	static const FName WrappingZ(TEXT("AddressZ"));
+	SetMetadataWithDisplayString(Metadata, InAsset, WrappingZ, TEXTVIEW("Z-axis Tiling Method"), MetaDataValueStringToDisplayString);
 
-	// Just use the save hash for now
-	FIoHash PackageHash = AssetPackageData.GetPackageSavedHash();
-	InHasher << PackageHash;
+	// Volume texture
+	static const FName WrappingMode(TEXT("AddressMode"));
+	SetMetadataWithDisplayString(Metadata, InAsset, WrappingMode, TEXTVIEW("Tiling Method"), MetaDataValueStringToDisplayString);
 
-	return true;
-}
+	static const FName CompressionSettings(TEXT("CompressionSettings"));
+	SetMetadataWithDisplayString(Metadata, InAsset, CompressionSettings, TEXTVIEW("Compression Settings"), MetaDataValueStringToDisplayString);
 
-void FTextureProcessor::GenerateCaptionRequest(const TSharedRef<const FAssetData>& InAsset, bool bCanDisruptUser, DerivedData::FRequestOwner& RequestOwner, FOnRequestComplete OnRequestComplete) const
-{
-	FString PackagePath;
-	if (FPackageName::DoesPackageExist(InAsset->PackageName.ToString(), &PackagePath))
-	{
-		TSet<FName> ObjectsFullName;
-		FName ObjectFullName = *InAsset->GetFullName();
-		ObjectsFullName.Add(ObjectFullName);
+	static const FName SRGBTag(TEXT("SRGB"));
+	SetMetadata(Metadata, InAsset, SRGBTag, TEXTVIEW("sRGB"));
 
-		FThumbnailMap ThumbnailMap;
-		// Don't use the cache to avoid pollution. Consider revisiting this?
-		if (ThumbnailTools::LoadThumbnailsFromPackageDirectly(PackagePath, ObjectsFullName, ThumbnailMap, PackagePath))
-		{
-			if (FObjectThumbnail* Thumbnail = ThumbnailMap.Find(ObjectFullName))
-			{
-				FCaptionRequest CaptionRequest;
-
-				// Convert the format from BGRA8 to RGBA8 inplace;
-				Thumbnail->DecompressImageData();
-				if (!ensure(Thumbnail->GetImage().Format == ERawImageFormat::BGRA8))
-				{
-					OnRequestComplete(false, FCaptionRequest(), FString(TEXT("Thumbnail was using an unexpected data format.")));
-					return;
-				}
-				TArrayView<uint8> UncompressedImage = Thumbnail->AccessImageData();
-				for (int32 PixelIndex = 0; PixelIndex < UncompressedImage.Num(); PixelIndex += 4)
-				{
-					constexpr int32 RedOffsetInBGRA8 = 2;
-					Swap(UncompressedImage[PixelIndex], UncompressedImage[PixelIndex + RedOffsetInBGRA8]);
-				}
-				Thumbnail->CompressImageData();
-				Thumbnail->AccessImageData().Empty();
-
-
-				FAssetMedia Media;
-				Media.MimeType = Thumbnail->GetCompressor()->GetMimeType();
-				Media.Data = MoveTemp(Thumbnail->AccessCompressedImageData());
-				CaptionRequest.AssetMedia.Add(MoveTemp(Media));
-
-				CaptionRequest.AssetPath = InAsset->GetObjectPathString();
-				CaptionRequest.AssetType = InAsset->AssetClassPath.ToString();
-
-				// Meta data
-				CaptionRequest.Metadata = MakeShared<FJsonObject>();
-				FString Value;
-
-				auto SetMetadata = [&CaptionRequest, &InAsset, &Value](FName AssetDataTag, FStringView CaptionMetadataKey)
-				{
-					if (InAsset->GetTagValue(AssetDataTag, Value))
-					{
-						CaptionRequest.Metadata->SetStringField(FString(CaptionMetadataKey), MoveTemp(Value));
-					}
-				};
-
-				auto SetMetadataWithDisplayString = [this, &CaptionRequest, &InAsset, &Value](FName AssetDataTag, FStringView CaptionMetadataKey)
-				{
-					if (InAsset->GetTagValue(AssetDataTag, Value))
-					{
-						const FString* DisplayValue = MetaDataValueStringToDisplayString.Find(Value);
-						CaptionRequest.Metadata->SetStringField(FString(CaptionMetadataKey), DisplayValue ? *DisplayValue : MoveTemp(Value));
-					}
-				};
-
-				static const FName DimensionsTag(TEXT("Dimensions"));
-				SetMetadata(DimensionsTag, TEXTVIEW("Dimensions"));
-
-				static const FName MaxTextureSizeTag(TEXT("MaxTextureSize"));
-				if (InAsset->GetTagValue(MaxTextureSizeTag, Value) && Value != TEXTVIEW("0"))
-				{
-					CaptionRequest.Metadata->SetStringField(FString(TEXTVIEW("MaxTextureSize")), MoveTemp(Value));
-				}
-
-				static const FName WrappingX(TEXT("AddressX"));
-				SetMetadataWithDisplayString(WrappingX, TEXTVIEW("X-axis Tiling Method"));
-
-				static const FName WrappingY(TEXT("AddressY"));
-				SetMetadataWithDisplayString(WrappingY, TEXTVIEW("Y-axis Tiling Method"));
-
-				// Texture arrays
-				static const FName WrappingZ(TEXT("AddressZ"));
-				SetMetadataWithDisplayString(WrappingZ, TEXTVIEW("Z-axis Tiling Method"));
-
-				// Volume texture
-				static const FName WrappingMode(TEXT("AddressMode"));
-				SetMetadataWithDisplayString(WrappingMode, TEXTVIEW("Tiling Method"));
-
-				static const FName CompressionSettings(TEXT("CompressionSettings"));
-				SetMetadataWithDisplayString(CompressionSettings, TEXTVIEW("Compression Settings"));
-
-				static const FName SRGBTag(TEXT("SRGB"));
-				SetMetadata(SRGBTag, TEXTVIEW("sRGB"));
-
-				static const FName VirtualTextureStreamingTag(TEXT("VirtualTextureStreaming"));
-				SetMetadata(VirtualTextureStreamingTag, TEXTVIEW("Virtual Texture Streaming"));
-
-
-				OnRequestComplete(true, MoveTemp(CaptionRequest), FString());
-				return;
-			}
-		}
-	}
+	static const FName VirtualTextureStreamingTag(TEXT("VirtualTextureStreaming"));
+	SetMetadata(Metadata, InAsset, VirtualTextureStreamingTag, TEXTVIEW("Virtual Texture Streaming"));
 
-	OnRequestComplete(false, FCaptionRequest(), FString(TEXT("Couldn't fetch an on-disk thumbnail for the asset")));
+	return Metadata;
 }
 
 FTextureProcessor::FTextureProcessor()
 {
-	// Reflection data is not thread safe in editor
-	check(IsInGameThread());
-
-	auto PopulateFromEnum = [this](const UEnum* Enum)
-	{
-		for (int32 Index = 0; Index < Enum->NumEnums() - 1; ++Index)
-		{
-			FString Key = Enum->GetNameStringByIndex(Index);
-
-			// We don't want the user localization in our metadata
-			FString DisplayValue = Enum->GetDisplayNameTextByIndex(Index).BuildSourceString();
-			MetaDataValueStringToDisplayString.Add(MoveTemp(Key), MoveTemp(DisplayValue));
-		}
-	};
-
-	PopulateFromEnum(StaticEnum<TextureAddress>());
-	PopulateFromEnum(StaticEnum<TextureCompressionSettings>());
+	PopulateFromEnum(MetaDataValueStringToDisplayString, StaticEnum<TextureAddress>());
+	PopulateFromEnum(MetaDataValueStringToDisplayString, StaticEnum<TextureCompressionSettings>());
 }
 
 }
diff --git a/Engine/Plugins/Experimental/SemanticSearch/Source/SemanticSearch/Private/AssetProcessors/TextureAssetProcessor.h b/Engine/Plugins/Experimental/SemanticSearch/Source/SemanticSearch/Private/AssetProcessors/TextureAssetProcessor.h
--- a/Engine/Plugins/Experimental/SemanticSearch/Source/SemanticSearch/Private/AssetProcessors/TextureAssetProcessor.h
+++ b/Engine/Plugins/Experimental/SemanticSearch/Source/SemanticSearch/Private/AssetProcessors/TextureAssetProcessor.h
@@ -2,26 +2,28 @@
 
 #pragma once
 
-#include "Interfaces/IAssetProcessor.h"
+#include "Implementations/ThumbnailBaseProcessor.h"
 
 namespace UE::SemanticSearch::Private
 {
-class FTextureProcessor : public IAssetProcessor
+
+class FTextureProcessor : public FThumbnailBaseAssetProcessor
 {
 public:
 	// Must be created on the game thread
 	FTextureProcessor();
 
+	constexpr virtual FStringView GetProcessSubBucketName() const override;
+
 	// Thread safe
-	virtual UClass& GetSupportedClass()  const override;
+	virtual UClass& GetSupportedClass() const override;
 	virtual bool SupportDerivedClasses() const override;
-	virtual bool GenerateAssetHash(const FAssetData& Asset, FMemoryHasherBlake3& InHasher) const override;
-	virtual void GenerateCaptionRequest(const TSharedRef<const FAssetData>& InAsset, bool bCanDisruptUser, DerivedData::FRequestOwner& RequestOwner, FOnRequestComplete OnRequestComplete) const override;
+
+	virtual TSharedPtr<FJsonObject> GetMetadata(const TSharedRef<const FAssetData>& InAsset) const override;
 
 private:
 	// Cache the display value as those are not thread safe
 	TMap<FString, FString> MetaDataValueStringToDisplayString;
+};
 
-	};
 }
-
diff --git a/Engine/Plugins/Experimental/SemanticSearch/Source/SemanticSearch/Private/Implementations/ThumbnailBaseProcessor.cpp b/Engine/Plugins/Experimental/SemanticSearch/Source/SemanticSearch/Private/Implementations/ThumbnailBaseProcessor.cpp
new file mode 100644
--- /dev/null
+++ b/Engine/Plugins/Experimental/SemanticSearch/Source/SemanticSearch/Private/Implementations/ThumbnailBaseProcessor.cpp
@@ -0,0 +1,100 @@
+// Copyright Epic Games, Inc. All Rights Reserved.
+
+#include "Implementations/ThumbnailBaseProcessor.h"
+
+#include "AssetRegistry/IAssetRegistry.h"
+#include "ImageCore.h"
+#include "Misc/Guid.h"
+#include "Misc/ObjectThumbnail.inl"
+#include "Misc/PackageName.h"
+#include "ObjectTools.h"
+
+namespace UE::SemanticSearch
+{
+
+bool FThumbnailBaseAssetProcessor::GenerateAssetHash(const FAssetData& Asset, FMemoryHasherBlake3& InHasher) const
+{
+	if (Asset.HasAnyPackageFlags(PKG_Cooked))
+	{
+		return false;
+	}
+
+	FAssetPackageData AssetPackageData;
+	if (IAssetRegistry::GetChecked().TryGetAssetPackageData(Asset.PackageName, AssetPackageData) != AssetRegistry::EExists::Exists)
+	{
+		return false;
+	}
+
+	constexpr FGuid DerivedDataGuid = FGuid(0x47702F04, 0x21B445B6, 0x96DE9103, 0xC16D689E);
+
+	// Safe we know the hasher won't modify it
+	InHasher << const_cast<FGuid&>(DerivedDataGuid);
+
+	FName PackageName(Asset.PackageName);
+	InHasher << PackageName;
+
+	FName AssetName(Asset.AssetName);
+	InHasher << AssetName;
+
+	FIoHash PackageHash = AssetPackageData.GetPackageSavedHash();
+	InHasher << PackageHash;
+
+	int32 Revision = GetRevision();
+	InHasher << Revision;
+
+	return true;
+}
+
+void FThumbnailBaseAssetProcessor::GenerateCaptionRequest(const TSharedRef<const FAssetData>& InAsset, bool bCanDisruptUser, DerivedData::FRequestOwner& RequestOwner, FOnRequestComplete OnRequestComplete) const
+{
+	FString PackagePath;
+	if (FPackageName::DoesPackageExist(InAsset->PackageName.ToString(), &PackagePath))
+	{
+		TSet<FName> ObjectsFullName;
+		FName ObjectFullName = *InAsset->GetFullName();
+		ObjectsFullName.Add(ObjectFullName);
+
+		FThumbnailMap ThumbnailMap;
+		// Don't use the cache to avoid pollution. Consider revisiting this?
+		if (ThumbnailTools::LoadThumbnailsFromPackageDirectly(PackagePath, ObjectsFullName, ThumbnailMap, PackagePath))
+		{
+			if (FObjectThumbnail* Thumbnail = ThumbnailMap.Find(ObjectFullName))
+			{
+				FCaptionRequest CaptionRequest;
+
+				// Convert the format from BGRA8 to RGBA8 inplace
+				Thumbnail->DecompressImageData();
+				if (!ensure(Thumbnail->GetImage().Format == ERawImageFormat::BGRA8))
+				{
+					OnRequestComplete(false, FCaptionRequest(), FString(TEXT("Thumbnail was using an unexpected data format.")));
+					return;
+				}
+				TArrayView<uint8> UncompressedImage = Thumbnail->AccessImageData();
+				for (int32 PixelIndex = 0; PixelIndex < UncompressedImage.Num(); PixelIndex += 4)
+				{
+					constexpr int32 RedOffsetInBGRA8 = 2;
+					Swap(UncompressedImage[PixelIndex], UncompressedImage[PixelIndex + RedOffsetInBGRA8]);
+				}
+				Thumbnail->CompressImageData();
+				Thumbnail->AccessImageData().Empty();
+
+				FAssetMedia Media;
+				Media.MimeType = Thumbnail->GetCompressor()->GetMimeType();
+				Media.Data = MoveTemp(Thumbnail->AccessCompressedImageData());
+				CaptionRequest.AssetMedia.Add(MoveTemp(Media));
+
+				CaptionRequest.AssetPath = InAsset->GetObjectPathString();
+				CaptionRequest.AssetType = InAsset->AssetClassPath.ToString();
+
+				CaptionRequest.Metadata = GetMetadata(InAsset);
+
+				OnRequestComplete(true, MoveTemp(CaptionRequest), FString());
+				return;
+			}
+		}
+	}
+
+	OnRequestComplete(false, FCaptionRequest(), FString(TEXT("Couldn't fetch an on-disk thumbnail for the asset")));
+}
+
+}
diff --git a/Engine/Plugins/Experimental/SemanticSearch/Source/SemanticSearch/Public/Implementations/ThumbnailBaseProcessor.h b/Engine/Plugins/Experimental/SemanticSearch/Source/SemanticSearch/Public/Implementations/ThumbnailBaseProcessor.h
new file mode 100644
--- /dev/null
+++ b/Engine/Plugins/Experimental/SemanticSearch/Source/SemanticSearch/Public/Implementations/ThumbnailBaseProcessor.h
@@ -0,0 +1,38 @@
+// Copyright Epic Games, Inc. All Rights Reserved.
+
+#pragma once
+
+#include "Interfaces/IAssetProcessor.h"
+
+#include "Serialization/MemoryHasher.h"
+#include "Templates/SharedPointer.h"
+
+class FJsonObject;
+struct FAssetData;
+
+namespace UE::SemanticSearch
+{
+
+/**
+ * Reusable base class for asset processors that operate via on-disk thumbnails.
+ *
+ * Derived classes must:
+ *  - Implement GetProcessSubBucketName() tell the system where to store the derived
+ *  - Implement GetMetadata(InAsset) to supply structured per-asset metadata.
+ *
+ * Derived classes may optionally override GetRevision() to bump cached data if they their implementation details.
+ */
+class SEMANTICSEARCH_API FThumbnailBaseAssetProcessor : public IAssetProcessor
+{
+public:
+	/** Returns structured metadata for the asset. Called from GenerateCaptionRequest. */
+	virtual TSharedPtr<FJsonObject> GetMetadata(const TSharedRef<const FAssetData>& InAsset) const = 0;
+
+	/** Bump this to invalidate previously cached derived data. */
+	virtual int32 GetRevision() const { return 1; }
+
+	virtual bool GenerateAssetHash(const FAssetData& Asset, FMemoryHasherBlake3& InHasher) const override;
+	virtual void GenerateCaptionRequest(const TSharedRef<const FAssetData>& InAsset, bool bCanDisruptUser, DerivedData::FRequestOwner& RequestOwner, FOnRequestComplete OnRequestComplete) const override;
+};
+
+}
diff --git a/Engine/Plugins/Experimental/SemanticSearch/Source/SemanticSearch/Public/Interfaces/IAssetProcessor.h b/Engine/Plugins/Experimental/SemanticSearch/Source/SemanticSearch/Public/Interfaces/IAssetProcessor.h
--- a/Engine/Plugins/Experimental/SemanticSearch/Source/SemanticSearch/Public/Interfaces/IAssetProcessor.h
+++ b/Engine/Plugins/Experimental/SemanticSearch/Source/SemanticSearch/Public/Interfaces/IAssetProcessor.h
@@ -25,6 +25,13 @@ using FOnRequestComplete = TUniqueFunction<void(bool bHasGeneratedCaptionRequest
 class IAssetProcessor
 {
 public:
+
+	/**
+	 * Used to separate the generated data into their own bucket
+	 * Use using the name of your type is acceptable value if it is unique
+	 */
+	constexpr virtual FStringView GetProcessSubBucketName() const = 0;
+
 	virtual UClass& GetSupportedClass() const = 0;
 	virtual bool SupportDerivedClasses() const = 0;
 
diff --git a/Engine/Plugins/Experimental/SemanticSearch/Source/SemanticSearch/SemanticSearch.Build.cs b/Engine/Plugins/Experimental/SemanticSearch/Source/SemanticSearch/SemanticSearch.Build.cs
--- a/Engine/Plugins/Experimental/SemanticSearch/Source/SemanticSearch/SemanticSearch.Build.cs
+++ b/Engine/Plugins/Experimental/SemanticSearch/Source/SemanticSearch/SemanticSearch.Build.cs
@@ -5,7 +5,7 @@ using UnrealBuildTool;
 /**
  * Core module for semantic search service
  * Provide the key system for the queries and management of the generation of the embeddings
- * 
+ *
  * Any ui or higher level logic should go in another module
  */
 public class SemanticSearch : ModuleRules
@@ -23,6 +23,7 @@ public class SemanticSearch : ModuleRules
 				"HTTP",
 				"ImageCore",
 				"Json",
+				"PhysicsCore",
 				"UnrealEd",
 			});
 
