Improve RemoteEmbeddingProvider HTTP request threading by setting delegate completion to run on the HTTP thread instead of the game thread, and dispatch the request complete callback in its own task to avoid blocking the HTTP request thread

#Synced-CL 51761360

#Synced-CL 51761360

diff --git a/Engine/Plugins/Experimental/SemanticSearch/Source/SemanticSearch/Private/RemoteEmbeddingProvider.cpp b/Engine/Plugins/Experimental/SemanticSearch/Source/SemanticSearch/Private/RemoteEmbeddingProvider.cpp
--- a/Engine/Plugins/Experimental/SemanticSearch/Source/SemanticSearch/Private/RemoteEmbeddingProvider.cpp
+++ b/Engine/Plugins/Experimental/SemanticSearch/Source/SemanticSearch/Private/RemoteEmbeddingProvider.cpp
@@ -38,22 +38,24 @@ public:
 
 	virtual void SetPriority(DerivedData::EPriority Priority) override
 	{
-		EHttpRequestPriority HttpPriority = EHttpRequestPriority::Low;
+		// Keep the priority lower the then actual tasks to not accidentally starve the ddc
+		EHttpRequestPriority HttpPriority = EHttpRequestPriority::Lowest;
 		switch (Priority)
 		{
 			case DerivedData::EPriority::Lowest:
 				HttpPriority = EHttpRequestPriority::Lowest;
 				break;
 			case DerivedData::EPriority::Low:
-				HttpPriority = EHttpRequestPriority::Low;
+				HttpPriority = EHttpRequestPriority::Lowest;
 				break;
 			case DerivedData::EPriority::Normal:
-				HttpPriority = EHttpRequestPriority::Normal;
+				HttpPriority = EHttpRequestPriority::Low;
 				break;
 			case DerivedData::EPriority::High:
-				HttpPriority = EHttpRequestPriority::High;
+				HttpPriority = EHttpRequestPriority::Normal;
 				break;
 			case DerivedData::EPriority::Highest:
+				HttpPriority = EHttpRequestPriority::High;
 			case DerivedData::EPriority::Blocking:
 				HttpPriority = EHttpRequestPriority::Highest;
 				break;
@@ -79,9 +81,14 @@ public:
 private:
 	void OnRequestComplete(FHttpRequestPtr Request, FHttpResponsePtr Response, bool bConnected)
 	{
+		// This will die after the scope 
 		RequestOwner->End(this, [this, &Request, &Response, bConnected]()
 			{
-				RequestCompleteDelegate.ExecuteIfBound(Request, Response, bConnected);
+				// Free the http request thread asap
+				RequestOwner->LaunchTask(TEXT("SemanticSearch_RemoteRequestComplete"), [InRequestCompleteDelegate = MoveTemp(RequestCompleteDelegate), Request, Response, bConnected]()
+					{
+						InRequestCompleteDelegate.ExecuteIfBound(Request, Response, bConnected);
+					});
 			});
 	}
 
@@ -144,6 +151,7 @@ void FRemoteEmbeddingProvider::GenerateCaptionAsync(
 		FString::Printf(TEXT("multipart/form-data; boundary=%s"), *Boundary));
 	Http->SetContent(MoveTemp(Body));
 	Http->SetTimeout(60.0f);
+	Http->SetDelegateThreadPolicy(EHttpRequestDelegateThreadPolicy::CompleteOnHttpThread);
 
 	Http->OnProcessRequestComplete().BindLambda(
 		[this, OnComplete = MakeShared<FOnCaptionComplete>(MoveTemp(InOnComplete))](FHttpRequestPtr /*Req*/, FHttpResponsePtr Resp, bool bConnected)
@@ -235,6 +243,7 @@ void FRemoteEmbeddingProvider::GenerateEmbeddingAsync(
 	Http->SetHeader(TEXT("Content-Type"), TEXT("application/json"));
 	Http->SetContentAsString(PayloadStr);
 	Http->SetTimeout(30.0f);
+	Http->SetDelegateThreadPolicy(EHttpRequestDelegateThreadPolicy::CompleteOnHttpThread);
 
 	Http->OnProcessRequestComplete().BindLambda(
 		[this, OnComplete = MakeShared<FOnEmbeddingComplete>(MoveTemp(InOnComplete))](FHttpRequestPtr /*Req*/, FHttpResponsePtr Resp, bool bConnected)
