Introduce an AssetProcessorManager singleton to register and retrieve asset processors, and register default built-in processors on module startup
Refactor the DDC cache task to use static factory methods with pre-generated cache keys, per-processor DDC sub-buckets, and a proactive CacheData path
Fix a race condition in request broadcasting by introducing atomic active-request flags and splitting broadcast logic into separate embedding and caption dispatch paths, and add CPU profiler trace scopes throughout.
Still need to cancel pending job on engne exit and remove the finished job from the container.

#Synced-CL 51986944

diff --git a/Engine/Plugins/Experimental/SemanticSearch/Source/SemanticSearch/Private/AssetProcessorManager.cpp b/Engine/Plugins/Experimental/SemanticSearch/Source/SemanticSearch/Private/AssetProcessorManager.cpp
new file mode 100644
--- /dev/null
+++ b/Engine/Plugins/Experimental/SemanticSearch/Source/SemanticSearch/Private/AssetProcessorManager.cpp
@@ -0,0 +1,213 @@
+// Copyright Epic Games, Inc. All Rights Reserved.
+
+#include "AssetProcessorManager.h"
+
+#include "Misc/ScopeRWLock.h"
+#include "SemanticSearchDerivedDataBuilder.h"
+#include "UObject/Class.h"
+#include "UObject/UObjectGlobals.h"
+#include "UObject/UObjectHash.h"
+
+namespace UE::SemanticSearch
+{
+	using namespace UE::DerivedData;
+	FAssetProcessorManager& FAssetProcessorManager::Get()
+	{
+		static FAssetProcessorManager ManagerSingleton;
+		return ManagerSingleton;
+	}
+
+	FAssetProcessorManager::FAssetProcessorManager()
+	{
+		OnObjectsReinstancedHandle = FCoreUObjectDelegates::OnObjectsReinstanced.AddRaw(this, &FAssetProcessorManager::OnObjectsReinstanced);
+		OnModuleUnloadedHandle = FCoreUObjectDelegates::CompiledInUObjectsRemovedDelegate.AddRaw(this, &FAssetProcessorManager::OnModuleUnloaded);
+		OnHotReloadCompleteHandle = FCoreUObjectDelegates::ReloadReinstancingCompleteDelegate.AddRaw(this, &FAssetProcessorManager::OnHotReloadComplete);
+	}
+
+	FAssetProcessorManager::~FAssetProcessorManager()
+	{
+		FCoreUObjectDelegates::OnObjectsReinstanced.Remove(OnObjectsReinstancedHandle);
+		FCoreUObjectDelegates::CompiledInUObjectsRemovedDelegate.Remove(OnModuleUnloadedHandle);
+		FCoreUObjectDelegates::ReloadReinstancingCompleteDelegate.Remove(OnHotReloadCompleteHandle);
+	}
+
+	bool FAssetProcessorManager::GetOrCreateTask(const FAssetData& InAssetData, TFunctionRef<void(const TSharedRef<Private::FSemanticSearchBuildCacheTask>&)> InApply)
+	{
+		TSharedPtr<IAssetProcessor> AssetProcessor;
+		{
+			FReadScopeLock ReadLock(AssetTypeToProcessorLock);
+			TSharedPtr<IAssetProcessor>* PtrToAssetProcessor = AssetTypeToProcessor.Find(InAssetData.AssetClassPath);
+			if (PtrToAssetProcessor)
+			{
+				AssetProcessor = *PtrToAssetProcessor;
+			}
+		}
+
+		if (!AssetProcessor)
+		{
+			return false;
+		}
+
+		FCacheKey CacheKey = Private::FSemanticSearchBuildCacheTask::GenerateCacheKey(InAssetData, AssetProcessor.ToSharedRef());
+		if (CacheKey == FCacheKey::Empty)
+		{
+			return false;
+		}
+
+		BuildsJobs.FindOrProduceAndApply(
+			CacheKey,
+			[&InAssetData, &AssetProcessor, &CacheKey]() -> TSharedRef<Private::FSemanticSearchBuildCacheTask>
+			{
+				return Private::FSemanticSearchBuildCacheTask::Create(InAssetData, AssetProcessor.ToSharedRef(), MoveTemp(CacheKey)).ToSharedRef();
+			},
+			InApply);
+
+		return true;
+	}
+
+	bool FAssetProcessorManager::StarCachingAsset(FAssetData&& InAssetData)
+	{
+		return GetOrCreateTask(InAssetData, [](const TSharedRef<Private::FSemanticSearchBuildCacheTask>& Task)
+		{
+			Task->CacheData();
+		});
+	}
+
+	bool FAssetProcessorManager::GetCaptionData(const FAssetData& InAssetData, TUniqueFunction<void(FAssetCaptionResult&&, FString&&)>&& InOnDataAvailable)
+	{
+		return GetOrCreateTask(InAssetData, [&InOnDataAvailable](const TSharedRef<Private::FSemanticSearchBuildCacheTask>& Task)
+		{
+			Task->GetCaptionData([OnDataAvailable = MoveTemp(InOnDataAvailable)](Private::FCaptionDerivedData&& CaptionData, FString&& ErrorMessage) mutable
+			{
+				FAssetCaptionResult Result;
+				Result.Caption   = MoveTemp(CaptionData.Caption);
+				Result.Keywords  = MoveTemp(CaptionData.Keywords);
+				OnDataAvailable(MoveTemp(Result), MoveTemp(ErrorMessage));
+			});
+		});
+	}
+
+	bool FAssetProcessorManager::GetEmbeddingData(const FAssetData& InAssetData, TUniqueFunction<void(FAssetEmbeddingResult&&, FString&&)>&& InOnDataAvailable)
+	{
+		return GetOrCreateTask(InAssetData, [&InOnDataAvailable](const TSharedRef<Private::FSemanticSearchBuildCacheTask>& Task)
+		{
+			Task->GetEmbeddingData([OnDataAvailable = MoveTemp(InOnDataAvailable)](Private::FEmbeddingDerivedData&& EmbeddingData, FString&& ErrorMessage) mutable
+			{
+				FAssetEmbeddingResult Result;
+				Result.Embedding = MoveTemp(EmbeddingData.Embedding);
+				OnDataAvailable(MoveTemp(Result), MoveTemp(ErrorMessage));
+			});
+		});
+	}
+
+	
+
+	void FAssetProcessorManager::RegisterAssetProcessor(TSharedRef<IAssetProcessor>&& AssetProcessor)
+	{
+		UClass& AssetClass = AssetProcessor->GetSupportedClass();
+
+		if (!ensureMsgf(AssetClass.HasAnyClassFlags(CLASS_Native), TEXT("The asset processor only support native asset type. The unreal is not really made to properly support non native asset types.")))
+		{
+			return;
+		}
+
+		const FTopLevelAssetPath ClassPath = AssetClass.GetClassPathName();
+		
+		TSharedPtr<IAssetProcessor> ExistingAssetProcessor;
+		{
+			FReadScopeLock ReadLock(AssetTypeToProcessorLock);
+			ExistingAssetProcessor = AssetTypeToProcessor.FindRef(ClassPath);
+		}
+
+		if (ExistingAssetProcessor)
+		{
+			UnregisterAssetProcessor(AssetClass);
+		}
+
+		{
+			FWriteScopeLock WriteLock(AssetTypeToProcessorLock);
+			AssetTypeToProcessor.Add(ClassPath, MoveTemp(AssetProcessor));
+		}
+
+		InvalidAssetChildClassesCache(AssetClass);
+
+	}
+
+	void FAssetProcessorManager::UnregisterAssetProcessor(UClass& Class)
+	{
+		{
+			FWriteScopeLock WriteLock(AssetTypeToProcessorLock);
+			AssetTypeToProcessor.Remove(Class.GetClassPathName());
+		}
+
+		InvalidAssetChildClassesCache(Class);
+	}
+
+	void FAssetProcessorManager::InvalidAssetChildClassesCache(UClass& Class)
+	{
+		// If slow due to non native maybe consider testing the handling the recursion by hand to stop early
+		constexpr bool bRecursive = true;
+		TArray<UClass*> DerivedClasses;
+		GetDerivedClasses(&Class, DerivedClasses, bRecursive);
+
+
+		FWriteScopeLock WriteLock(AssetTypeToProcessorLock);
+		for (UClass* ChildClass : DerivedClasses)
+		{
+			if (ChildClass && ChildClass->HasAnyClassFlags(CLASS_Native))
+			{
+				AssetTypeToProcessor.Remove(ChildClass->GetClassPathName());
+			}
+		}
+	}
+
+	void FAssetProcessorManager::OnObjectsReinstanced(const TMap<UObject*, UObject*>& InReplacementObjects)
+	{
+		// todo rewrite properly
+		/**FWriteScopeLock WriteLock(AssetTypeToProcessorLock);
+		for (TPair<FTopLevelAssetPath, TSharedPtr<IAssetProcessor>> Pair : AssetTypeToProcessor)
+		{
+			if (UObject* const* NewObject = InReplacementObjects.Find(FindObject<UClass>(nullptr, *Pair.Key.ToString())))
+			{
+				if (*NewObject)
+				{
+					if (const UClass* NewClass = Cast<const UClass>(*NewObject))
+					{
+						if (NewClass->HasAnyClassFlags(CLASS_Native))
+						{
+							AssetTypeToProcessor.Add(NewClass->GetClassPathName(), MoveTemp(Pair.Value));
+						}
+					}
+				}
+				AssetTypeToProcessor.Remove(FTopLevelAssetPath(*NewObject));
+			}
+		}*/
+	}
+
+	void FAssetProcessorManager::OnModuleUnloaded(TConstArrayView<UPackage*> UnloadingPackages)
+	{
+		FWriteScopeLock WriteLock(AssetTypeToProcessorLock);
+		for (auto It = AssetTypeToProcessor.CreateIterator(); It ; ++It)
+		{
+			if (const UClass* Class = FindObject<UClass>(nullptr, *It->Key.ToString()))
+			{
+				if (UnloadingPackages.Contains(Class->GetPackage()))
+				{
+					It.RemoveCurrent();
+				}
+			}
+		}
+	}
+
+	void FAssetProcessorManager::OnHotReloadComplete()
+	{
+		FWriteScopeLock WriteLock(AssetTypeToProcessorLock);
+		for (auto It = AssetTypeToProcessor.CreateIterator(); It; ++It)
+		{
+			if (!FindObject<UClass>(nullptr, *It->Key.ToString()))
+			{
+				It.RemoveCurrent();
+			}
+		}
+	}
+}
\ No newline at end of file
diff --git a/Engine/Plugins/Experimental/SemanticSearch/Source/SemanticSearch/Private/AssetProcessors/AssetProcessorUtils.cpp b/Engine/Plugins/Experimental/SemanticSearch/Source/SemanticSearch/Private/AssetProcessors/AssetProcessorUtils.cpp
--- a/Engine/Plugins/Experimental/SemanticSearch/Source/SemanticSearch/Private/AssetProcessors/AssetProcessorUtils.cpp
+++ b/Engine/Plugins/Experimental/SemanticSearch/Source/SemanticSearch/Private/AssetProcessors/AssetProcessorUtils.cpp
@@ -2,6 +2,13 @@
 
 #include "AssetProcessors/AssetProcessorUtils.h"
 
+#include "AssetProcessorManager.h"
+#include "AssetProcessors/MaterialAssetProcessor.h"
+#include "AssetProcessors/MaterialInstanceAssetProcessor.h"
+#include "AssetProcessors/SkeletalMeshAssetProcessor.h"
+#include "AssetProcessors/StaticMeshAssetProcessor.h"
+#include "AssetProcessors/TextureAssetProcessor.h"
+
 #include "AssetRegistry/AssetData.h"
 #include "Dom/JsonObject.h"
 #include "UObject/Class.h"
@@ -52,4 +59,13 @@ void PopulateFromEnum(TMap<FString, FString>& OutMap, const UEnum* Enum)
 	}
 }
 
+void RegisterDefaultAssetProcessors()
+{
+	UE::SemanticSearch::FAssetProcessorManager& Manager = UE::SemanticSearch::FAssetProcessorManager::Get();
+	Manager.RegisterAssetProcessor<FMaterialProcessor>();
+	Manager.RegisterAssetProcessor<FMaterialInstanceProcessor>();
+	Manager.RegisterAssetProcessor<FStaticMeshProcessor>();
+	Manager.RegisterAssetProcessor<FSkeletalMeshProcessor>();
+	Manager.RegisterAssetProcessor<FTextureProcessor>();
+}
 }
diff --git a/Engine/Plugins/Experimental/SemanticSearch/Source/SemanticSearch/Private/AssetProcessors/AssetProcessorUtils.h b/Engine/Plugins/Experimental/SemanticSearch/Source/SemanticSearch/Private/AssetProcessors/AssetProcessorUtils.h
--- a/Engine/Plugins/Experimental/SemanticSearch/Source/SemanticSearch/Private/AssetProcessors/AssetProcessorUtils.h
+++ b/Engine/Plugins/Experimental/SemanticSearch/Source/SemanticSearch/Private/AssetProcessors/AssetProcessorUtils.h
@@ -51,4 +51,9 @@ void SetMetadataWithDisplayString(
  */
 void PopulateFromEnum(TMap<FString, FString>& OutMap, const UEnum* Enum);
 
+/**
+ * Register the build in asset processor that ship with this plugin
+ */
+void RegisterDefaultAssetProcessors();
+
 }
diff --git a/Engine/Plugins/Experimental/SemanticSearch/Source/SemanticSearch/Private/SemanticSearchDerivedDataBuilder.cpp b/Engine/Plugins/Experimental/SemanticSearch/Source/SemanticSearch/Private/SemanticSearchDerivedDataBuilder.cpp
--- a/Engine/Plugins/Experimental/SemanticSearch/Source/SemanticSearch/Private/SemanticSearchDerivedDataBuilder.cpp
+++ b/Engine/Plugins/Experimental/SemanticSearch/Source/SemanticSearch/Private/SemanticSearchDerivedDataBuilder.cpp
@@ -3,6 +3,8 @@
 #include "SemanticSearchDerivedDataBuilder.h"
 
 #include "Containers/SharedString.h"
+#include "ProfilingDebugging/CpuProfilerTrace.h"
+#include "Containers/StringFwd.h"
 #include "DerivedDataCache.h"
 #include "Interfaces/IAssetProcessor.h"
 #include "Misc/Guid.h"
@@ -12,6 +14,7 @@
 #include "Serialization/MemoryReader.h"
 #include "Serialization/MemoryWriter.h"
 #include "String/Find.h"
+#include "Misc/StringBuilder.h"
 
 namespace UE::SemanticSearch::Private
 {
@@ -44,29 +47,36 @@ DerivedData::ECachePolicy GetCachePolicy()
 	return ECachePolicy::Default | ECachePolicy::NonDeterministic;
 }
 
-FSemanticSearchBuildCacheTask::FSemanticSearchBuildCacheTask(const FAssetData& InAssetData, const TSharedRef<IAssetProcessor>& InAssetProcessor)
-	: DerivedData::FRequestOwner(DerivedData::EPriority::Low)
-	, Asset(MakeShared<FAssetData>(InAssetData))
-	, AssetProcessor(InAssetProcessor)
-	, AssetKey(GenerateCacheKey())
+TSharedPtr<FSemanticSearchBuildCacheTask> FSemanticSearchBuildCacheTask::Create(const FAssetData& InAssetData, const TSharedRef<IAssetProcessor>& InAssetProcessor, DerivedData::FCacheKey&& InKey)
 {
+	return MakeShared<FSemanticSearchBuildCacheTask>(FAssetData(InAssetData), InAssetProcessor, MoveTemp(InKey));
 }
 
-FSemanticSearchBuildCacheTask::FSemanticSearchBuildCacheTask(FAssetData&& InAssetData, const TSharedRef<IAssetProcessor>& InAssetProcessor)
+TSharedPtr<FSemanticSearchBuildCacheTask> FSemanticSearchBuildCacheTask::Create(FAssetData&& InAssetData, const TSharedRef<IAssetProcessor>& InAssetProcessor, DerivedData::FCacheKey&& InKey)
+{
+	return MakeShared<FSemanticSearchBuildCacheTask>(MoveTemp(InAssetData), InAssetProcessor, MoveTemp(InKey));
+}
+
+
+FSemanticSearchBuildCacheTask::FSemanticSearchBuildCacheTask(FAssetData&& InAssetData, const TSharedRef<IAssetProcessor>& InAssetProcessor, DerivedData::FCacheKey&& InAssetKey)
 	: DerivedData::FRequestOwner(DerivedData::EPriority::Low)
 	, Asset(MakeShared<FAssetData>(MoveTemp(InAssetData)))
 	, AssetProcessor(InAssetProcessor)
-	, AssetKey(GenerateCacheKey())
+	, AssetKey(MoveTemp(InAssetKey))
+	, DDCRequestName(FString::Format(TEXT("DDCSemanticSearch:%s"), {FStringFormatArg(Asset->GetObjectPathString())}))
 {
 }
 
 void FSemanticSearchBuildCacheTask::GetEmbeddingData(TUniqueFunction<void(FEmbeddingDerivedData&& EmbeddingData, FString&& ErrorMessage)>&& InOnDataAvailable)
 {
+	TRACE_CPUPROFILER_EVENT_SCOPE_STR("SemanticSearch::GetEmbeddingData");
 	{
 		FReadScopeLock ReadLock(IsBuildingDerivedDataLock);
-		if (bIsBuildingDerivedData)
+		QueuedEmbeddingRequests.Enqueue(MoveTemp(InOnDataAvailable));
+		bool bPreviousValue = false;
+		HasActiveEmbeddingRequest.compare_exchange_strong(bPreviousValue, true);
+		if (bIsBuildingDerivedData || bPreviousValue)
 		{
-			QueuedEmbeddingRequests.Enqueue(MoveTemp(InOnDataAvailable));
 			return;
 		}
 	}
@@ -74,82 +84,122 @@ void FSemanticSearchBuildCacheTask::GetEmbeddingData(TUniqueFunction<void(FEmbed
 	using namespace DerivedData;
 
 	FCacheGetChunkRequest Request;
+	Request.Name = DDCRequestName;
 	Request.Key = AssetKey;
 	Request.Id = GetEmbeddingDatachunkId();
 	Request.Policy = GetCachePolicy();
 
-	GetCache().GetChunks({ Request }, *this, [OnDataAvailable = MoveTemp(InOnDataAvailable), this](FCacheGetChunkResponse&& Response) mutable
+	FRequestBarrier Barrier(*this);
+	GetCache().GetChunks({ Request }, *this, [this](FCacheGetChunkResponse&& Response) mutable
 		{
 			if (Response.Status == EStatus::Ok && Response.RawData.GetSize() > 0)
 			{
+				TRACE_CPUPROFILER_EVENT_SCOPE_STR("SemanticSearch::GetEmbeddingData::DDCHit");
 				FEmbeddingDerivedData Data;
-				
+
 				FMemoryReaderView Reader(Response.RawData.GetView(), true);
 				SerializeEmbedding(Data, Reader);
 				Response.RawData.Reset();
-				
-				OnDataAvailable(MoveTemp(Data), FString());
+
+				BroadcastQueuedEmbeddingRequest(MoveTemp(Data), FString());
 			}
 			else if (Response.Status != EStatus::Canceled)
 			{
-				QueuedEmbeddingRequests.Enqueue(MoveTemp(OnDataAvailable));
+				TRACE_CPUPROFILER_EVENT_SCOPE_STR("SemanticSearch::GetEmbeddingData::DDCMiss");
 				this->BuildDerivedData();
 			}
 			else
 			{
-				// Something canceled the request
-				OnDataAvailable(FEmbeddingDerivedData(), FString(TEXT("The data request was cancelled")));
+				TRACE_CPUPROFILER_EVENT_SCOPE_STR("SemanticSearch::GetEmbeddingData::DDCCanceled");
+				BroadcastQueuedEmbeddingRequest(FEmbeddingDerivedData(), FString(TEXT("The data request was cancelled")));
 			}
 		});
 }
 
 void FSemanticSearchBuildCacheTask::GetCaptionData(TUniqueFunction<void(FCaptionDerivedData&& CaptionData, FString&& ErrorMessage)>&& InOnDataAvailable)
 {
+	TRACE_CPUPROFILER_EVENT_SCOPE_STR("SemanticSearch::GetCaptionData");
 	{
 		FReadScopeLock ReadLock(IsBuildingDerivedDataLock);
-		if (bIsBuildingDerivedData)
+		QueuedCaptionRequests.Enqueue(MoveTemp(InOnDataAvailable));
+		bool bPreviousValue = false;
+		HasActiveCaptionRequest.compare_exchange_strong(bPreviousValue, true);
+		if (bIsBuildingDerivedData || bPreviousValue)
 		{
-			QueuedCaptionRequests.Enqueue(MoveTemp(InOnDataAvailable));
 			return;
 		}
 	}
 
 	using namespace DerivedData;
 	FCacheGetChunkRequest Request;
+	Request.Name = DDCRequestName;
 	Request.Key = AssetKey;
 	Request.Id = GetCaptionDatachunkId();
 	Request.Policy = GetCachePolicy();
 
-	GetCache().GetChunks({ Request }, *this, [OnDataAvailable = MoveTemp(InOnDataAvailable), this](FCacheGetChunkResponse&& Response) mutable
+	FRequestBarrier Barrier(*this);
+	GetCache().GetChunks({ Request }, *this, [this](FCacheGetChunkResponse&& Response) mutable
 		{
 			if (Response.Status == EStatus::Ok && Response.RawData.GetSize() > 0)
 			{
+				TRACE_CPUPROFILER_EVENT_SCOPE_STR("SemanticSearch::GetCaptionData::DDCHit");
 				FCaptionDerivedData Data;
 				FMemoryReaderView Reader(Response.RawData.GetView(), true);
 				SerializeCaption(Data, Reader);
 				Response.RawData.Reset();
 
-				OnDataAvailable(MoveTemp(Data), FString());
+				BroadcastQueuedCaptionRequest(MoveTemp(Data), FString());
 			}
 			else if (Response.Status != EStatus::Canceled)
 			{
-				QueuedCaptionRequests.Enqueue(MoveTemp(OnDataAvailable));
+				TRACE_CPUPROFILER_EVENT_SCOPE_STR("SemanticSearch::GetCaptionData::DDCMiss");
 				BuildDerivedData();
 			}
 			else
 			{
-				// Something canceled the request
-				OnDataAvailable(FCaptionDerivedData(), FString(TEXT("The data request was cancelled")));
+				TRACE_CPUPROFILER_EVENT_SCOPE_STR("SemanticSearch::GetCaptionData::DDCCanceled");
+				BroadcastQueuedCaptionRequest(FCaptionDerivedData(), FString(TEXT("The data request was cancelled")));
 			}
 		});
 }
 
-DerivedData::FCacheKey FSemanticSearchBuildCacheTask::GenerateCacheKey() const
+void FSemanticSearchBuildCacheTask::CacheData()
 {
-	static FGuid SemanticSearchDerivedDataGuid(TEXT("A3D4A6CE-B6C0-4E41-ABC5-B4C2999D0095"));
+	TRACE_CPUPROFILER_EVENT_SCOPE_STR("SemanticSearch::CacheData");
+	{
+		FReadScopeLock ReadLock(IsBuildingDerivedDataLock);
+		if (bIsBuildingDerivedData)
+		{
+			return;
+		}
+	}
+
+	using namespace DerivedData;
+
+	FCacheGetRequest Request;
+	Request.Name = DDCRequestName;
+	Request.Key = AssetKey;
+	Request.Policy = GetCachePolicy() | ECachePolicy::SkipData;
+
+	FRequestBarrier Barrier(*this);
+	GetCache().Get({ Request }, *this, [this](FCacheGetResponse&& Response) mutable
+		{
+			if (Response.Status != EStatus::Ok && Response.Status != EStatus::Canceled)
+			{
+				TRACE_CPUPROFILER_EVENT_SCOPE_STR("SemanticSearch::CachedData::DDCMiss");
+				BuildDerivedData();
+			}
+		});
+}
+
+DerivedData::FCacheKey FSemanticSearchBuildCacheTask::GenerateCacheKey(const FAssetData& InAssetData, const TSharedRef<IAssetProcessor>& InAssetProcessor)
+{
+	constexpr FGuid SemanticSearchDerivedDataGuid(0xA3D4A6CE, 0xB6C04E41, 0xABC5B4C2, 0x999D0095);
 	FMemoryHasherBlake3 HashWriter;
-	HashWriter << SemanticSearchDerivedDataGuid;
-	if (!AssetProcessor->GenerateAssetHash(*Asset, HashWriter))
+
+	// The Hasher won't modify the value
+	HashWriter << const_cast<FGuid&>(SemanticSearchDerivedDataGuid);
+	if (!InAssetProcessor->GenerateAssetHash(InAssetData, HashWriter))
 	{
 		// Invalid Key Generation
 		return DerivedData::FCacheKey{};
@@ -159,30 +209,48 @@ DerivedData::FCacheKey FSemanticSearchBuildCacheTask::GenerateCacheKey() const
 
 	DerivedData::FCacheKey CacheKey;
 	CacheKey.Hash = HashWriter.Finalize();
-	CacheKey.Bucket = DerivedData::FCacheBucket(TEXT("SemanticSearch"));
+	CacheKey.Bucket = GetCacheKeyBucket(InAssetProcessor);
 	return CacheKey;
 }
 
+DerivedData::FCacheBucket FSemanticSearchBuildCacheTask::GetCacheKeyBucket(const TSharedRef<IAssetProcessor>& InAssetProcessor)
+{
+	TStringBuilder<256> Builder;
+	Builder.Append(TEXTVIEW("SemanticSearch"));
+	Builder.Append(InAssetProcessor->GetProcessSubBucketName());
+	return DerivedData::FCacheBucket(FStringView(Builder));
+}
+
+
 void FSemanticSearchBuildCacheTask::BuildDerivedData()
 {
+	TRACE_CPUPROFILER_EVENT_SCOPE_STR("SemanticSearch::BuildDerivedData");
 	{
 		FWriteScopeLock WriteLock(IsBuildingDerivedDataLock);
+		if (bIsBuildingDerivedData)
+		{
+			return;
+		}
+
 		bIsBuildingDerivedData = true;
 	}
 
 	AssetProcessor->GenerateCaptionRequest(Asset, false, *this, [this](bool bHasGeneratedRequest, FCaptionRequest&& CaptionRequest, FString&& ErrorMessage)
 		{
+			TRACE_CPUPROFILER_EVENT_SCOPE_STR("SemanticSearch::BuildDerivedData::CaptionRequestReady");
 			if (bHasGeneratedRequest)
 			{
 				if (TSharedPtr<IEmbeddingProvider> EmbeddingProvider = FSemanticSearchModule::Get().GetRegisteredEmbeddingProvider())
 				{
 					EmbeddingProvider->GenerateCaptionAsync(MoveTemp(CaptionRequest), [this, EmbeddingProvider](FCaptionResponse&& Caption)
 						{
+							TRACE_CPUPROFILER_EVENT_SCOPE_STR("SemanticSearch::BuildDerivedData::CaptionResponseReceived");
 							if (Caption.ErrorMessage.IsEmpty())
 							{
 								FString CaptionString = Caption.Caption;
 								EmbeddingProvider->GenerateEmbeddingAsync(CaptionString, [this, UsedCaption = MoveTemp(Caption)](FEmbeddingResponse&& Embedding) mutable
 									{
+										TRACE_CPUPROFILER_EVENT_SCOPE_STR("SemanticSearch::BuildDerivedData::EmbeddingResponseReceived");
 										if (Embedding.ErrorMessage.IsEmpty())
 										{
 											FCaptionDerivedData CaptionDerivedData;
@@ -191,7 +259,7 @@ void FSemanticSearchBuildCacheTask::BuildDerivedData()
 											FEmbeddingDerivedData EmbeddingDerivedData;
 											EmbeddingDerivedData.Embedding = Embedding.Embedding;
 
-											if (AssetKey == GenerateCacheKey())
+											if (AssetKey == GenerateCacheKey(*Asset, AssetProcessor))
 											{
 												using namespace DerivedData;
 
@@ -213,14 +281,14 @@ void FSemanticSearchBuildCacheTask::BuildDerivedData()
 													RecordBuilder.AddValue(GetEmbeddingDatachunkId(), MakeSharedBufferFromArray(MoveTemp(Buffer)));
 												}
 
-												DerivedData::FCachePutRequest PutRequest{ FSharedString(FString::Format(TEXT("%s SemanticSearchData"), {FStringFormatArg(Asset->GetObjectPathString())})), RecordBuilder.Build(), GetCachePolicy() | ECachePolicy::SkipData };
-
+												DerivedData::FCachePutRequest PutRequest{ DDCRequestName, RecordBuilder.Build(), GetCachePolicy() | ECachePolicy::SkipData };
 
 												DerivedData::GetCache().Put(
 													{ PutRequest },
 													*this,
 													[AssetPtr = Asset](DerivedData::FCachePutResponse&& Response)
 													{
+														TRACE_CPUPROFILER_EVENT_SCOPE_STR("SemanticSearch::BuildDerivedData::DDCPutResponse");
 														if (Response.Status == EStatus::Ok)
 														{
 															UE_LOGF(LogSemanticSearch, Log, "Successfully stored semantic data in DDC for: %ls", *AssetPtr->GetObjectPathString());
@@ -232,7 +300,6 @@ void FSemanticSearchBuildCacheTask::BuildDerivedData()
 													});
 
 												BroadcastQueuedRequests(MoveTemp(EmbeddingDerivedData), MoveTemp(CaptionDerivedData), FString());
-
 											}
 											else
 											{
@@ -267,54 +334,41 @@ void FSemanticSearchBuildCacheTask::BuildDerivedData()
 }
 
 
-void FSemanticSearchBuildCacheTask::BroadcastQueuedRequests(FEmbeddingDerivedData&& EmbeddingData, FCaptionDerivedData&& CaptionData, FString&& ErrorMessage)
+TUniqueFunction<void(FEmbeddingDerivedData&&, FString&&)> FSemanticSearchBuildCacheTask::NoLock_GetEmbeddingRequest()
 {
-	TArray<TUniqueFunction<void(FEmbeddingDerivedData&&, FString&&)>, TInlineAllocator<4>> EmbeddingRequests;
-	TArray<TUniqueFunction<void(FCaptionDerivedData&&, FString&&)>, TInlineAllocator<4>> CaptionRequests;
-
+	TArray<TUniqueFunction<void(FEmbeddingDerivedData&&, FString&&)>, TInlineAllocator<4>> Callbacks;
+	bool bPreviousValue = true;
+	HasActiveEmbeddingRequest.compare_exchange_strong(bPreviousValue, false);
+	if (bPreviousValue)
 	{
-		FWriteScopeLock WriteLock(IsBuildingDerivedDataLock);
-		bIsBuildingDerivedData = false;
-
-		// Grab Embedding Requests
+		TUniqueFunction<void(FEmbeddingDerivedData&&, FString&&)> Callback;
+		while (QueuedEmbeddingRequests.Dequeue(Callback))
 		{
-			TUniqueFunction<void(FEmbeddingDerivedData&&, FString&&)> Callback;
-			while (QueuedEmbeddingRequests.Dequeue(Callback))
-			{
-				EmbeddingRequests.Add(MoveTemp(Callback));
-			}
-		}
-
-
-		// Grab Caption Request
-		{
-			TUniqueFunction<void(FCaptionDerivedData&&, FString&&)> Callback;
-			while (QueuedCaptionRequests.Dequeue(Callback))
-			{
-				CaptionRequests.Add(MoveTemp(Callback));
-			}
+			Callbacks.Add(MoveTemp(Callback));
 		}
 	}
 
+	if (Callbacks.IsEmpty())
+	{
+		return {};
+	}
 
-	FString CopiedErrorMessage;
-	FString* ErrorMessageToUse = &ErrorMessage;
-
-	// Broadcast Embedding Requests
+	return [Callbacks = MoveTemp(Callbacks)](FEmbeddingDerivedData&& EmbeddingData, FString&& ErrorMessage) mutable
 	{
 		FEmbeddingDerivedData CopiedEmbeddingData;
 		FEmbeddingDerivedData* EmbeddingDataToUse = &EmbeddingData;
+		FString CopiedErrorMessage;
+		FString* ErrorMessageToUse = &ErrorMessage;
 		int32 Count = 0;
 
-		for (TUniqueFunction<void(FEmbeddingDerivedData&&, FString&&)>& Callback : EmbeddingRequests)
+		for (TUniqueFunction<void(FEmbeddingDerivedData&&, FString&&)>& Callback : Callbacks)
 		{
 			++Count;
 			if (Callback)
 			{
-				bool bSwapCurrentToCopied = false;
-				if (Count < EmbeddingRequests.Num())
+				bool bSwapCurrentToCopied = Count < Callbacks.Num();
+				if (bSwapCurrentToCopied)
 				{
-					bSwapCurrentToCopied = true;
 					CopiedEmbeddingData = *EmbeddingDataToUse;
 					CopiedErrorMessage = *ErrorMessageToUse;
 				}
@@ -332,23 +386,44 @@ void FSemanticSearchBuildCacheTask::BroadcastQueuedRequests(FEmbeddingDerivedDat
 				}
 			}
 		}
+	};
+}
+
+TUniqueFunction<void(FCaptionDerivedData&&, FString&&)> FSemanticSearchBuildCacheTask::NoLock_GetCaptionRequests()
+{
+	TArray<TUniqueFunction<void(FCaptionDerivedData&&, FString&&)>, TInlineAllocator<4>> Callbacks;
+	bool bPreviousValue = true;
+	HasActiveCaptionRequest.compare_exchange_strong(bPreviousValue, false);
+	if (bPreviousValue)
+	{
+		TUniqueFunction<void(FCaptionDerivedData&&, FString&&)> Callback;
+		while (QueuedCaptionRequests.Dequeue(Callback))
+		{
+			Callbacks.Add(MoveTemp(Callback));
+		}
+	}
+
+	if (Callbacks.IsEmpty())
+	{
+		return {};
 	}
 
-	// Broadcast Caption Request
+	return [Callbacks = MoveTemp(Callbacks)](FCaptionDerivedData&& CaptionData, FString&& ErrorMessage) mutable
 	{
 		FCaptionDerivedData CopiedCaptionData;
 		FCaptionDerivedData* CaptionDataToUse = &CaptionData;
+		FString CopiedErrorMessage;
+		FString* ErrorMessageToUse = &ErrorMessage;
 		int32 Count = 0;
 
-		for (TUniqueFunction<void(FCaptionDerivedData&&, FString&&)>& Callback : CaptionRequests)
+		for (TUniqueFunction<void(FCaptionDerivedData&&, FString&&)>& Callback : Callbacks)
 		{
 			++Count;
 			if (Callback)
 			{
-				bool bSwapCurrentToCopied = false;
-				if (Count < CaptionRequests.Num())
+				bool bSwapCurrentToCopied = Count < Callbacks.Num();
+				if (bSwapCurrentToCopied)
 				{
-					bSwapCurrentToCopied = true;
 					CopiedCaptionData = *CaptionDataToUse;
 					CopiedErrorMessage = *ErrorMessageToUse;
 				}
@@ -366,6 +441,66 @@ void FSemanticSearchBuildCacheTask::BroadcastQueuedRequests(FEmbeddingDerivedDat
 				}
 			}
 		}
+	};
+}
+
+void FSemanticSearchBuildCacheTask::BroadcastQueuedEmbeddingRequest(FEmbeddingDerivedData&& EmbeddingData, FString&& ErrorMessage)
+{
+	TRACE_CPUPROFILER_EVENT_SCOPE_STR("SemanticSearch::BroadcastEmbeddingRequest");
+	TUniqueFunction<void(FEmbeddingDerivedData&&, FString&&)> EmbeddingCallback;
+	{
+		FWriteScopeLock WriteLock(IsBuildingDerivedDataLock);
+		EmbeddingCallback = NoLock_GetEmbeddingRequest();
+	}
+
+	if (EmbeddingCallback)
+	{
+		TRACE_CPUPROFILER_EVENT_SCOPE_STR("SemanticSearch::BroadcastEmbeddingRequest::Dispatch");
+		EmbeddingCallback(MoveTemp(EmbeddingData), MoveTemp(ErrorMessage));
+	}
+}
+
+void FSemanticSearchBuildCacheTask::BroadcastQueuedCaptionRequest(FCaptionDerivedData&& CaptionData, FString&& ErrorMessage)
+{
+	TRACE_CPUPROFILER_EVENT_SCOPE_STR("SemanticSearch::BroadcastCaptionRequest");
+	TUniqueFunction<void(FCaptionDerivedData&&, FString&&)> CaptionCallback;
+	{
+		FWriteScopeLock WriteLock(IsBuildingDerivedDataLock);
+		CaptionCallback = NoLock_GetCaptionRequests();
+	}
+
+	if (CaptionCallback)
+	{
+		TRACE_CPUPROFILER_EVENT_SCOPE_STR("SemanticSearch::BroadcastCaptionRequest::Dispatch");
+		CaptionCallback(MoveTemp(CaptionData), MoveTemp(ErrorMessage));
+	}
+}
+
+void FSemanticSearchBuildCacheTask::BroadcastQueuedRequests(FEmbeddingDerivedData&& EmbeddingData, FCaptionDerivedData&& CaptionData, FString&& ErrorMessage)
+{
+	TRACE_CPUPROFILER_EVENT_SCOPE_STR("SemanticSearch::BroadcastAllRequests");
+	TUniqueFunction<void(FEmbeddingDerivedData&&, FString&&)> EmbeddingCallback;
+	TUniqueFunction<void(FCaptionDerivedData&&, FString&&)> CaptionCallback;
+
+	{
+		FWriteScopeLock WriteLock(IsBuildingDerivedDataLock);
+		bIsBuildingDerivedData = false;
+		EmbeddingCallback = NoLock_GetEmbeddingRequest();
+		CaptionCallback = NoLock_GetCaptionRequests();
+	}
+
+	if (EmbeddingCallback)
+	{
+		TRACE_CPUPROFILER_EVENT_SCOPE_STR("SemanticSearch::BroadcastAllRequests::DispatchEmbedding");
+		// Copy the error message if it will also be needed by CaptionCallback
+		FString EmbeddingErrorMessage = CaptionCallback ? ErrorMessage : MoveTemp(ErrorMessage);
+		EmbeddingCallback(MoveTemp(EmbeddingData), MoveTemp(EmbeddingErrorMessage));
+	}
+
+	if (CaptionCallback)
+	{
+		TRACE_CPUPROFILER_EVENT_SCOPE_STR("SemanticSearch::BroadcastAllRequests::DispatchCaption");
+		CaptionCallback(MoveTemp(CaptionData), MoveTemp(ErrorMessage));
 	}
 }
 
diff --git a/Engine/Plugins/Experimental/SemanticSearch/Source/SemanticSearch/Private/SemanticSearchDerivedDataBuilder.h b/Engine/Plugins/Experimental/SemanticSearch/Source/SemanticSearch/Private/SemanticSearchDerivedDataBuilder.h
--- a/Engine/Plugins/Experimental/SemanticSearch/Source/SemanticSearch/Private/SemanticSearchDerivedDataBuilder.h
+++ b/Engine/Plugins/Experimental/SemanticSearch/Source/SemanticSearch/Private/SemanticSearchDerivedDataBuilder.h
@@ -3,6 +3,7 @@
 #pragma once
 
 #include "AssetRegistry/AssetData.h"
+#include "Containers/SharedString.h"
 #include "Containers/Array.h"
 #include "Containers/Queue.h"
 #include "Containers/UnrealString.h"
@@ -10,7 +11,7 @@
 #include "DerivedDataRequestOwner.h"
 #include "HAL/CriticalSection.h"
 #include "Templates/SharedPointer.h"
-
+#include <atomic>
 
 namespace UE::SemanticSearch
 {
@@ -37,27 +38,43 @@ struct FCaptionDerivedData
 class FSemanticSearchBuildCacheTask : public DerivedData::FRequestOwner
 {
 public:
-	FSemanticSearchBuildCacheTask(const FAssetData& InAssetData, const TSharedRef<IAssetProcessor>& InAssetProcessor);
-	FSemanticSearchBuildCacheTask(FAssetData&& InAssetData, const TSharedRef<IAssetProcessor>& InAssetProcessor);
+	// InKey must not be empty — passing an empty key will crash. Use GenerateCacheKey and validate before calling.
+	static TSharedPtr<FSemanticSearchBuildCacheTask> Create(const FAssetData& InAssetData, const TSharedRef<IAssetProcessor>& InAssetProcessor, DerivedData::FCacheKey&& InKey);
+	static TSharedPtr<FSemanticSearchBuildCacheTask> Create(FAssetData&& InAssetData, const TSharedRef<IAssetProcessor>& InAssetProcessor, DerivedData::FCacheKey&& InKey);
+
+	FSemanticSearchBuildCacheTask(FAssetData&& InAssetData, const TSharedRef<IAssetProcessor>& InAssetProcessor, DerivedData::FCacheKey&& InAssetKey);
 
 	void GetEmbeddingData(TUniqueFunction<void (FEmbeddingDerivedData&& EmbeddingData, FString&& ErrorMessage)>&& InOnDataAvailable);
 	void GetCaptionData(TUniqueFunction<void (FCaptionDerivedData&& CaptionData, FString&& ErrorMessage)>&& InOnDataAvailable);
 
+	void CacheData();
+
+	static DerivedData::FCacheKey GenerateCacheKey(const FAssetData& InAssetData, const TSharedRef<IAssetProcessor>& InAssetProcessor);
+	static DerivedData::FCacheBucket GetCacheKeyBucket(const TSharedRef<IAssetProcessor>& InAssetProcessor);
+
 private:
-	DerivedData::FCacheKey GenerateCacheKey() const;
 
 	void BuildDerivedData();
 	void BroadcastQueuedRequests(FEmbeddingDerivedData&& EmbeddingData, FCaptionDerivedData&& CaptionData, FString&& ErrorMessage);
+	void BroadcastQueuedCaptionRequest(FCaptionDerivedData&& CaptionData, FString&& ErrorMessage);
+	void BroadcastQueuedEmbeddingRequest(FEmbeddingDerivedData&& EmbeddingData, FString&& ErrorMessage);
+
+	TUniqueFunction<void(FCaptionDerivedData&&, FString&&)> NoLock_GetCaptionRequests();
+	TUniqueFunction<void(FEmbeddingDerivedData&&, FString&&)> NoLock_GetEmbeddingRequest();
+
 
 	TSharedRef<const FAssetData> Asset;
 	TSharedRef<IAssetProcessor> AssetProcessor;
 	DerivedData::FCacheKey AssetKey;
+	FSharedString DDCRequestName;
 
 	// Acs has a lock for the IsBuildingDerivedData flag but also avoid race condition with the queued Requests
 	FRWLock IsBuildingDerivedDataLock;
 	bool bIsBuildingDerivedData = false;
 	TQueue<TUniqueFunction<void (FEmbeddingDerivedData&&, FString&&)>, EQueueMode::Mpsc> QueuedEmbeddingRequests;
 	TQueue<TUniqueFunction<void (FCaptionDerivedData&&, FString&&)>, EQueueMode::Mpsc> QueuedCaptionRequests;
+	std::atomic<bool> HasActiveEmbeddingRequest = false;
+	std::atomic<bool> HasActiveCaptionRequest = false;
 };
 }
 }
diff --git a/Engine/Plugins/Experimental/SemanticSearch/Source/SemanticSearch/Private/SemanticSearchModule.cpp b/Engine/Plugins/Experimental/SemanticSearch/Source/SemanticSearch/Private/SemanticSearchModule.cpp
--- a/Engine/Plugins/Experimental/SemanticSearch/Source/SemanticSearch/Private/SemanticSearchModule.cpp
+++ b/Engine/Plugins/Experimental/SemanticSearch/Source/SemanticSearch/Private/SemanticSearchModule.cpp
@@ -2,6 +2,8 @@
 
 #include "SemanticSearchModule.h"
 
+#include "AssetProcessors/AssetProcessorUtils.h"
+
 #include "Modules/ModuleManager.h"
 #include "RemoteEmbeddingProvider.h"
 
@@ -21,6 +23,7 @@ void FSemanticSearchModule::StartupModule()
 {
 	UE_LOG(LogSemanticSearch, Log, TEXT("SemanticSearch module starting up"));
 	RegisteredEmbeddingProvider = MakeShared<FRemoteEmbeddingProvider>();
+	RegisterDefaultAssetProcessors();
 }
 
 void FSemanticSearchModule::ShutdownModule()
diff --git a/Engine/Plugins/Experimental/SemanticSearch/Source/SemanticSearch/Public/AssetProcessorManager.h b/Engine/Plugins/Experimental/SemanticSearch/Source/SemanticSearch/Public/AssetProcessorManager.h
new file mode 100644
--- /dev/null
+++ b/Engine/Plugins/Experimental/SemanticSearch/Source/SemanticSearch/Public/AssetProcessorManager.h
@@ -0,0 +1,95 @@
+// Copyright Epic Games, Inc. All Rights Reserved.
+
+#pragma once
+
+#include "Containers/Array.h"
+#include "Containers/Map.h"
+#include "Containers/StripedMap.h"
+#include "Containers/UnrealString.h"
+#include "Delegates/IDelegateInstance.h"
+#include "DerivedDataCacheKey.h"
+#include "HAL/CriticalSection.h"
+#include "Interfaces/IAssetProcessor.h"
+#include "Templates/Function.h"
+#include "Templates/UniquePtr.h"
+#include "Templates/UnrealTypeTraits.h"
+#include "UObject/TopLevelAssetPath.h"
+
+namespace UE::SemanticSearch
+{
+
+namespace Private
+{
+class FSemanticSearchBuildCacheTask;
+}
+
+struct FAssetCaptionResult
+{
+	FString Caption;
+	TArray<FString> Keywords;
+};
+
+struct FAssetEmbeddingResult
+{
+	TArray<float> Embedding;
+};
+
+class FAssetProcessorManager
+{
+public:
+
+	SEMANTICSEARCH_API static FAssetProcessorManager& Get();
+
+	SEMANTICSEARCH_API bool StarCachingAsset(FAssetData&& InAssetData);
+	SEMANTICSEARCH_API bool GetCaptionData(const FAssetData& InAssetData, TUniqueFunction<void(FAssetCaptionResult&&, FString&&)>&& InOnDataAvailable);
+	SEMANTICSEARCH_API bool GetEmbeddingData(const FAssetData& InAssetData, TUniqueFunction<void(FAssetEmbeddingResult&&, FString&&)>&& InOnDataAvailable);
+
+	template<typename FAssetProcessor>
+	void RegisterAssetProcessor();
+
+	template<typename FAssetProcessor>
+	void UnregisterAssetProcessor();
+
+
+private:
+
+	FAssetProcessorManager();
+	~FAssetProcessorManager();
+
+	SEMANTICSEARCH_API void RegisterAssetProcessor(TSharedRef<IAssetProcessor>&& AssetProcessor);
+	SEMANTICSEARCH_API void UnregisterAssetProcessor(UClass& Class);
+	void InvalidAssetChildClassesCache(UClass& Class);
+
+	bool GetOrCreateTask(const FAssetData& InAssetData, TFunctionRef<void(const TSharedRef<Private::FSemanticSearchBuildCacheTask>&)> InApply);
+
+	void OnObjectsReinstanced(const TMap<UObject*, UObject*>& InReplacementObjects);
+	void OnModuleUnloaded(TConstArrayView<UPackage*> UnloadingPackages);
+	void OnHotReloadComplete();
+
+	FDelegateHandle OnObjectsReinstancedHandle;
+	FDelegateHandle OnModuleUnloadedHandle;
+	FDelegateHandle OnHotReloadCompleteHandle;
+
+	FRWLock AssetTypeToProcessorLock;
+	TMap<FTopLevelAssetPath, TSharedPtr<IAssetProcessor>> AssetTypeToProcessor;
+
+	TStripedMap<32, DerivedData::FCacheKey, TSharedRef<Private::FSemanticSearchBuildCacheTask>> BuildsJobs;
+
+};
+
+template<typename FAssetProcessor>
+void FAssetProcessorManager::RegisterAssetProcessor()
+{
+	static_assert(TIsDerivedFrom<FAssetProcessor, IAssetProcessor>::Value, "Asset Processors must derive from IAssetProcessor");
+	RegisterAssetProcessor(MakeShared<FAssetProcessor>());
+}
+
+template<typename FAssetProcessor>
+void FAssetProcessorManager::UnregisterAssetProcessor()
+{
+	static_assert(TIsDerivedFrom<FAssetProcessor, IAssetProcessor>::Value, "Asset Processors must derive from IAssetProcessor");
+	FAssetProcessor AssetProcessor;
+	UnregisterAssetProcessor(AssetProcessor.SupportDerivedClasses(), AssetProcessor.GetSupportedClass());
+}
+
+}
\ No newline at end of file
