Fix the implementation of AssetProcessorManager::OnObjectsReinstanced by replacing the commented-out stub with a correct implementation
Use an iterator-based loop with RemoveCurrent to safely modify the map during iteration, and defer additions via a TChunkedArray to avoid invalidating the iterator

#Synced-CL 52060916

diff --git a/Engine/Plugins/Experimental/SemanticSearch/Source/SemanticSearch/Private/AssetProcessorManager.cpp b/Engine/Plugins/Experimental/SemanticSearch/Source/SemanticSearch/Private/AssetProcessorManager.cpp
--- a/Engine/Plugins/Experimental/SemanticSearch/Source/SemanticSearch/Private/AssetProcessorManager.cpp
+++ b/Engine/Plugins/Experimental/SemanticSearch/Source/SemanticSearch/Private/AssetProcessorManager.cpp
@@ -163,11 +163,11 @@ namespace UE::SemanticSearch
 
 	void FAssetProcessorManager::OnObjectsReinstanced(const TMap<UObject*, UObject*>& InReplacementObjects)
 	{
-		// todo rewrite properly
-		/**FWriteScopeLock WriteLock(AssetTypeToProcessorLock);
-		for (TPair<FTopLevelAssetPath, TSharedPtr<IAssetProcessor>> Pair : AssetTypeToProcessor)
+		FWriteScopeLock WriteLock(AssetTypeToProcessorLock);
+		TChunkedArray<TPair<FTopLevelAssetPath, TSharedPtr<IAssetProcessor>>> ProcessorsToAdd;
+		for (auto It = AssetTypeToProcessor.CreateIterator(); It; ++It)
 		{
-			if (UObject* const* NewObject = InReplacementObjects.Find(FindObject<UClass>(nullptr, *Pair.Key.ToString())))
+			if (UObject* const* NewObject = InReplacementObjects.Find(FindObject<UClass>(nullptr, *It->Key.ToString())))
 			{
 				if (*NewObject)
 				{
@@ -175,13 +175,19 @@ namespace UE::SemanticSearch
 					{
 						if (NewClass->HasAnyClassFlags(CLASS_Native))
 						{
-							AssetTypeToProcessor.Add(NewClass->GetClassPathName(), MoveTemp(Pair.Value));
+							ProcessorsToAdd.AddElement(TPair<FTopLevelAssetPath, TSharedPtr<IAssetProcessor>>(NewClass->GetClassPathName(), MoveTemp(It->Value)));
 						}
 					}
 				}
-				AssetTypeToProcessor.Remove(FTopLevelAssetPath(*NewObject));
+				It.RemoveCurrent();
 			}
-		}*/
+		}
+
+		AssetTypeToProcessor.Reserve(AssetTypeToProcessor.Num() + ProcessorsToAdd.Num());
+		for (TPair<FTopLevelAssetPath, TSharedPtr<IAssetProcessor>>& Pair : ProcessorsToAdd)
+		{
+			AssetTypeToProcessor.Add(MoveTemp(Pair.Key), MoveTemp(Pair.Value));
+		}
 	}
 
 	void FAssetProcessorManager::OnModuleUnloaded(TConstArrayView<UPackage*> UnloadingPackages)
