Abstract vector index interface + FAISS integration into DDC pipeline

- IVectorIndex: abstract interface with Train, Add, Remove, Search, Encode
- FPQVectorIndex: Product Quantization compressed search
- FFlatVectorIndex: Brute-force L2 search, no training needed
- DDC builder encodes embeddings (PQ codes) and adds vectors to index
- AssetProcessorManager owns the index
- GetAssetSemanticID: stable int64 hash for asset identification

#Synced-CL 52058676

diff --git a/Engine/Plugins/Experimental/SemanticSearch/Source/SemanticSearch/Private/AssetProcessorManager.cpp b/Engine/Plugins/Experimental/SemanticSearch/Source/SemanticSearch/Private/AssetProcessorManager.cpp
--- a/Engine/Plugins/Experimental/SemanticSearch/Source/SemanticSearch/Private/AssetProcessorManager.cpp
+++ b/Engine/Plugins/Experimental/SemanticSearch/Source/SemanticSearch/Private/AssetProcessorManager.cpp
@@ -100,7 +100,15 @@ namespace UE::SemanticSearch
 		});
 	}
 
-	
+	void FAssetProcessorManager::SetVectorIndex(TSharedPtr<IVectorIndex> InIndex)
+	{
+		VectorIndex = MoveTemp(InIndex);
+	}
+
+	TSharedPtr<IVectorIndex> FAssetProcessorManager::GetVectorIndex() const
+	{
+		return VectorIndex;
+	}
 
 	void FAssetProcessorManager::RegisterAssetProcessor(TSharedRef<IAssetProcessor>&& AssetProcessor)
 	{
diff --git a/Engine/Plugins/Experimental/SemanticSearch/Source/SemanticSearch/Private/AssetProcessors/AssetProcessorUtils.cpp b/Engine/Plugins/Experimental/SemanticSearch/Source/SemanticSearch/Private/AssetProcessors/AssetProcessorUtils.cpp
--- a/Engine/Plugins/Experimental/SemanticSearch/Source/SemanticSearch/Private/AssetProcessors/AssetProcessorUtils.cpp
+++ b/Engine/Plugins/Experimental/SemanticSearch/Source/SemanticSearch/Private/AssetProcessors/AssetProcessorUtils.cpp
@@ -68,4 +68,17 @@ void RegisterDefaultAssetProcessors()
 	Manager.RegisterAssetProcessor<FSkeletalMeshProcessor>();
 	Manager.RegisterAssetProcessor<FTextureProcessor>();
 }
+
+int64 GetAssetIndexID(const FAssetData& Asset)
+{
+	const FString Path = Asset.GetObjectPathString();
+	FBlake3 Hasher;
+	Hasher.Update(Path.GetCharArray().GetData(), Path.Len() * sizeof(TCHAR));
+	FBlake3Hash Hash = Hasher.Finalize();
+	int64 Result;
+	FMemory::Memcpy(&Result, Hash.GetBytes(), sizeof(int64));
+	// Ensure non-negative (FAISS uses -1 as sentinel)
+	return Result & 0x7FFFFFFFFFFFFFFF;
+}
+
 }
diff --git a/Engine/Plugins/Experimental/SemanticSearch/Source/SemanticSearch/Private/AssetProcessors/AssetProcessorUtils.h b/Engine/Plugins/Experimental/SemanticSearch/Source/SemanticSearch/Private/AssetProcessors/AssetProcessorUtils.h
--- a/Engine/Plugins/Experimental/SemanticSearch/Source/SemanticSearch/Private/AssetProcessors/AssetProcessorUtils.h
+++ b/Engine/Plugins/Experimental/SemanticSearch/Source/SemanticSearch/Private/AssetProcessors/AssetProcessorUtils.h
@@ -56,4 +56,10 @@ void PopulateFromEnum(TMap<FString, FString>& OutMap, const UEnum* Enum);
  */
 void RegisterDefaultAssetProcessors();
 
+/**
+ * Compute a stable int64 hash of an asset's object path, suitable for use as a vector index ID.
+ * Uses Blake3. Always non-negative (FAISS reserves -1 as sentinel).
+ */
+int64 GetAssetIndexID(const FAssetData& Asset);
+
 }
diff --git a/Engine/Plugins/Experimental/SemanticSearch/Source/SemanticSearch/Private/SemanticSearchDerivedDataBuilder.cpp b/Engine/Plugins/Experimental/SemanticSearch/Source/SemanticSearch/Private/SemanticSearchDerivedDataBuilder.cpp
--- a/Engine/Plugins/Experimental/SemanticSearch/Source/SemanticSearch/Private/SemanticSearchDerivedDataBuilder.cpp
+++ b/Engine/Plugins/Experimental/SemanticSearch/Source/SemanticSearch/Private/SemanticSearchDerivedDataBuilder.cpp
@@ -2,11 +2,14 @@
 
 #include "SemanticSearchDerivedDataBuilder.h"
 
+#include "AssetProcessorManager.h"
+#include "AssetProcessors/AssetProcessorUtils.h"
 #include "Containers/SharedString.h"
 #include "ProfilingDebugging/CpuProfilerTrace.h"
 #include "Containers/StringFwd.h"
 #include "DerivedDataCache.h"
 #include "Interfaces/IAssetProcessor.h"
+#include "Interfaces/IVectorIndex.h"
 #include "Misc/Guid.h"
 #include "Misc/ScopeRWLock.h"
 #include "SemanticSearchModule.h"
@@ -29,11 +32,21 @@ DerivedData::FValueId GetCaptionDatachunkId()
 	return DerivedData::FValueId::FromName(TEXT("CaptionData"));
 }
 
+DerivedData::FValueId GetQuantizedEmbeddingDatachunkId()
+{
+	return DerivedData::FValueId::FromName(TEXT("QuantizedEmbeddingData"));
+}
+
 void SerializeEmbedding(FEmbeddingDerivedData& Embedding, FArchive& Archive)
 {
 	Archive << Embedding.Embedding;
 }
 
+void SerializeQuantizedEmbedding(FQuantizedEmbeddingDerivedData& QuantizedEmbedding, FArchive& Archive)
+{
+	Archive << QuantizedEmbedding.QuantizedEmbedding;
+}
+
 void SerializeCaption(FCaptionDerivedData& Caption, FArchive& Archive)
 {
 	Archive << Caption.Caption;
@@ -259,6 +272,16 @@ void FSemanticSearchBuildCacheTask::BuildDerivedData()
 											FEmbeddingDerivedData EmbeddingDerivedData;
 											EmbeddingDerivedData.Embedding = Embedding.Embedding;
 
+											// Quantize the embedding if the active index supports it (PQ).
+											FQuantizedEmbeddingDerivedData QuantizedEmbeddingData;
+											if (TSharedPtr<IVectorIndex> VectorIndex = FAssetProcessorManager::Get().GetVectorIndex())
+											{
+												if (VectorIndex->SupportsQuantization() && VectorIndex->IsTrained())
+												{
+													QuantizedEmbeddingData.QuantizedEmbedding = VectorIndex->Quantize(EmbeddingDerivedData.Embedding);
+												}
+											}
+
 											if (AssetKey == GenerateCacheKey(*Asset, AssetProcessor))
 											{
 												using namespace DerivedData;
@@ -281,6 +304,15 @@ void FSemanticSearchBuildCacheTask::BuildDerivedData()
 													RecordBuilder.AddValue(GetEmbeddingDatachunkId(), MakeSharedBufferFromArray(MoveTemp(Buffer)));
 												}
 
+												// Save quantized embedding (PQ codes) if available
+												if (QuantizedEmbeddingData.QuantizedEmbedding.Num() > 0)
+												{
+													TArray<uint8> Buffer;
+													FMemoryWriter Writer(Buffer);
+													SerializeQuantizedEmbedding(QuantizedEmbeddingData, Writer);
+													RecordBuilder.AddValue(GetQuantizedEmbeddingDatachunkId(), MakeSharedBufferFromArray(MoveTemp(Buffer)));
+												}
+
 												DerivedData::FCachePutRequest PutRequest{ DDCRequestName, RecordBuilder.Build(), GetCachePolicy() | ECachePolicy::SkipData };
 
 												DerivedData::GetCache().Put(
@@ -299,6 +331,13 @@ void FSemanticSearchBuildCacheTask::BuildDerivedData()
 														}
 													});
 
+												// Add the embedding to the vector index
+												if (TSharedPtr<IVectorIndex> VectorIndex = FAssetProcessorManager::Get().GetVectorIndex())
+												{
+													const int64 AssetID = GetAssetIndexID(*Asset);
+													VectorIndex->Add(AssetID, EmbeddingDerivedData.Embedding);
+												}
+
 												BroadcastQueuedRequests(MoveTemp(EmbeddingDerivedData), MoveTemp(CaptionDerivedData), FString());
 											}
 											else
diff --git a/Engine/Plugins/Experimental/SemanticSearch/Source/SemanticSearch/Private/SemanticSearchDerivedDataBuilder.h b/Engine/Plugins/Experimental/SemanticSearch/Source/SemanticSearch/Private/SemanticSearchDerivedDataBuilder.h
--- a/Engine/Plugins/Experimental/SemanticSearch/Source/SemanticSearch/Private/SemanticSearchDerivedDataBuilder.h
+++ b/Engine/Plugins/Experimental/SemanticSearch/Source/SemanticSearch/Private/SemanticSearchDerivedDataBuilder.h
@@ -29,6 +29,11 @@ struct FEmbeddingDerivedData
 	TArray<float> Embedding;
 };
 
+struct FQuantizedEmbeddingDerivedData
+{
+	TArray<uint8> QuantizedEmbedding;
+};
+
 struct FCaptionDerivedData
 {
 	FString Caption;
diff --git a/Engine/Plugins/Experimental/SemanticSearch/Source/SemanticSearch/Private/SemanticSearchModule.cpp b/Engine/Plugins/Experimental/SemanticSearch/Source/SemanticSearch/Private/SemanticSearchModule.cpp
--- a/Engine/Plugins/Experimental/SemanticSearch/Source/SemanticSearch/Private/SemanticSearchModule.cpp
+++ b/Engine/Plugins/Experimental/SemanticSearch/Source/SemanticSearch/Private/SemanticSearchModule.cpp
@@ -2,7 +2,9 @@
 
 #include "SemanticSearchModule.h"
 
+#include "AssetProcessorManager.h"
 #include "AssetProcessors/AssetProcessorUtils.h"
+#include "VectorIndex/FlatVectorIndex.h"
 
 #include "Modules/ModuleManager.h"
 #include "RemoteEmbeddingProvider.h"
@@ -24,6 +26,11 @@ void FSemanticSearchModule::StartupModule()
 	UE_LOG(LogSemanticSearch, Log, TEXT("SemanticSearch module starting up"));
 	RegisteredEmbeddingProvider = MakeShared<FRemoteEmbeddingProvider>();
 	RegisterDefaultAssetProcessors();
+
+	// Default to a flat index. Dimension is set when first embedding arrives.
+	// Callers can replace this with a PQ index via SetVectorIndex().
+	constexpr int32 DefaultDimension = 768;
+	FAssetProcessorManager::Get().SetVectorIndex(MakeShared<FFlatVectorIndex>(DefaultDimension));
 }
 
 void FSemanticSearchModule::ShutdownModule()
@@ -41,5 +48,10 @@ void FSemanticSearchModule::RegisterEmbeddingProvider(TUniquePtr<IEmbeddingProvi
 	RegisteredEmbeddingProvider = MakeShareable<IEmbeddingProvider>(EmbeddingProvider.Release());
 }
 
+void FSemanticSearchModule::SetVectorIndex(TSharedPtr<IVectorIndex> InIndex)
+{
+	FAssetProcessorManager::Get().SetVectorIndex(MoveTemp(InIndex));
+}
+
 }
 
diff --git a/Engine/Plugins/Experimental/SemanticSearch/Source/SemanticSearch/Private/SemanticSearchModule.h b/Engine/Plugins/Experimental/SemanticSearch/Source/SemanticSearch/Private/SemanticSearchModule.h
--- a/Engine/Plugins/Experimental/SemanticSearch/Source/SemanticSearch/Private/SemanticSearchModule.h
+++ b/Engine/Plugins/Experimental/SemanticSearch/Source/SemanticSearch/Private/SemanticSearchModule.h
@@ -24,6 +24,7 @@ public:
 
 
 	virtual void RegisterEmbeddingProvider(TUniquePtr<IEmbeddingProvider>&& EmbeddingProvider) override;
+	virtual void SetVectorIndex(TSharedPtr<IVectorIndex> InIndex) override;
 
 private:
 	TSharedPtr<IEmbeddingProvider> RegisteredEmbeddingProvider;
diff --git a/Engine/Plugins/Experimental/SemanticSearch/Source/SemanticSearch/Private/VectorIndex/FlatVectorIndex.cpp b/Engine/Plugins/Experimental/SemanticSearch/Source/SemanticSearch/Private/VectorIndex/FlatVectorIndex.cpp
new file mode 100644
--- /dev/null
+++ b/Engine/Plugins/Experimental/SemanticSearch/Source/SemanticSearch/Private/VectorIndex/FlatVectorIndex.cpp
@@ -0,0 +1,110 @@
+// Copyright Epic Games, Inc. All Rights Reserved.
+
+#include "FlatVectorIndex.h"
+
+THIRD_PARTY_INCLUDES_START
+#include "faiss/index_io.h"
+#include "faiss/impl/io.h"
+THIRD_PARTY_INCLUDES_END
+
+namespace UE::SemanticSearch
+{
+
+FFlatVectorIndex::FFlatVectorIndex(int32 InDimension)
+	: Dimension(InDimension)
+{
+	faiss::IndexFlatL2* FlatIndex = new faiss::IndexFlatL2(Dimension);
+	IndexWithIDs = new faiss::IndexIDMap(FlatIndex);
+}
+
+FFlatVectorIndex::FFlatVectorIndex(int32 InDimension, faiss::IndexIDMap* InIndex)
+	: Dimension(InDimension)
+	, IndexWithIDs(InIndex)
+{
+}
+
+FFlatVectorIndex::~FFlatVectorIndex()
+{
+	delete IndexWithIDs;
+}
+
+void FFlatVectorIndex::Train(TConstArrayView<float> Vectors, int64 NumVectors)
+{
+	// Flat index does not need training.
+	return;
+}
+
+void FFlatVectorIndex::Add(int64 ID, TConstArrayView<float> Vector)
+{
+	check(Vector.Num() == Dimension);
+	faiss::idx_t FaissID = static_cast<faiss::idx_t>(ID);
+	IndexWithIDs->add_with_ids(1, Vector.GetData(), &FaissID);
+}
+
+bool FFlatVectorIndex::Remove(int64 ID)
+{
+	faiss::idx_t FaissID = static_cast<faiss::idx_t>(ID);
+	faiss::IDSelectorBatch Selector(1, &FaissID);
+	return IndexWithIDs->remove_ids(Selector) > 0;
+}
+
+TArray<FSearchResult> FFlatVectorIndex::Search(TConstArrayView<float> QueryVector, int32 K) const
+{
+	check(QueryVector.Num() == Dimension);
+
+	TArray<float> Distances;
+	TArray<faiss::idx_t> Labels;
+	Distances.SetNumUninitialized(K);
+	Labels.SetNumUninitialized(K);
+
+	IndexWithIDs->search(1, QueryVector.GetData(), K, Distances.GetData(), Labels.GetData());
+
+	TArray<FSearchResult> Results;
+	Results.Reserve(K);
+	for (int32 i = 0; i < K; ++i)
+	{
+		if (Labels[i] >= 0)
+		{
+			FSearchResult Result;
+			Result.ID = static_cast<int64>(Labels[i]);
+			Result.Distance = Distances[i];
+			Results.Add(Result);
+		}
+	}
+	return Results;
+}
+
+int64 FFlatVectorIndex::GetCount() const
+{
+	return IndexWithIDs->ntotal;
+}
+
+TArray<uint8> FFlatVectorIndex::Serialize() const
+{
+	faiss::VectorIOWriter Writer;
+	faiss::write_index(IndexWithIDs, &Writer);
+
+	TArray<uint8> Result;
+	Result.Append(
+		reinterpret_cast<const uint8*>(Writer.data.data()),
+		static_cast<int32>(Writer.data.size()));
+	return Result;
+}
+
+TUniquePtr<FFlatVectorIndex> FFlatVectorIndex::Deserialize(TConstArrayView<uint8> Data, int32 ExpectedDimension)
+{
+	faiss::VectorIOReader Reader;
+	Reader.data.assign(Data.GetData(), Data.GetData() + Data.Num());
+
+	faiss::Index* LoadedIndex = faiss::read_index(&Reader);
+	if (!LoadedIndex || LoadedIndex->d != ExpectedDimension)
+	{
+		delete LoadedIndex;
+		return nullptr;
+	}
+
+	faiss::IndexIDMap* IDMapIndex = static_cast<faiss::IndexIDMap*>(LoadedIndex);
+	return TUniquePtr<FFlatVectorIndex>(new FFlatVectorIndex(ExpectedDimension, IDMapIndex));
+}
+
+} // namespace UE::SemanticSearch
diff --git a/Engine/Plugins/Experimental/SemanticSearch/Source/SemanticSearch/Private/VectorIndex/FlatVectorIndex.h b/Engine/Plugins/Experimental/SemanticSearch/Source/SemanticSearch/Private/VectorIndex/FlatVectorIndex.h
new file mode 100644
--- /dev/null
+++ b/Engine/Plugins/Experimental/SemanticSearch/Source/SemanticSearch/Private/VectorIndex/FlatVectorIndex.h
@@ -0,0 +1,51 @@
+// Copyright Epic Games, Inc. All Rights Reserved.
+
+#pragma once
+
+#include "Interfaces/IVectorIndex.h"
+#include "Templates/UniquePtr.h"
+
+THIRD_PARTY_INCLUDES_START
+#include "faiss/IndexFlat.h"
+#include "faiss/IndexIDMap.h"
+THIRD_PARTY_INCLUDES_END
+
+namespace UE::SemanticSearch
+{
+
+/**
+ * Flat (brute-force) vector index using L2 distance.
+ *
+ * Stores raw vectors without compression. Exact search results.
+ * No training required. Good for small datasets or ground-truth comparison.
+ */
+class FFlatVectorIndex : public IVectorIndex
+{
+public:
+	explicit FFlatVectorIndex(int32 InDimension);
+	virtual ~FFlatVectorIndex() override;
+
+	FFlatVectorIndex(const FFlatVectorIndex&) = delete;
+	FFlatVectorIndex& operator=(const FFlatVectorIndex&) = delete;
+
+	// IVectorIndex
+	virtual void Train(TConstArrayView<float> Vectors, int64 NumVectors) override;
+	virtual bool IsTrained() const override { return true; }
+	virtual void Add(int64 ID, TConstArrayView<float> Vector) override;
+	virtual bool Remove(int64 ID) override;
+	virtual TArray<FSearchResult> Search(TConstArrayView<float> QueryVector, int32 K) const override;
+	virtual int64 GetCount() const override;
+	virtual int32 GetDimension() const override { return Dimension; }
+	virtual TArray<uint8> Serialize() const override;
+
+	/** Deserialize a full flat index. */
+	static TUniquePtr<FFlatVectorIndex> Deserialize(TConstArrayView<uint8> Data, int32 ExpectedDimension);
+
+private:
+	FFlatVectorIndex(int32 InDimension, faiss::IndexIDMap* InIndex);
+
+	int32 Dimension;
+	faiss::IndexIDMap* IndexWithIDs = nullptr;
+};
+
+} // namespace UE::SemanticSearch
diff --git a/Engine/Plugins/Experimental/SemanticSearch/Source/SemanticSearch/Private/VectorIndex/PQVectorIndex.cpp b/Engine/Plugins/Experimental/SemanticSearch/Source/SemanticSearch/Private/VectorIndex/PQVectorIndex.cpp
new file mode 100644
--- /dev/null
+++ b/Engine/Plugins/Experimental/SemanticSearch/Source/SemanticSearch/Private/VectorIndex/PQVectorIndex.cpp
@@ -0,0 +1,204 @@
+// Copyright Epic Games, Inc. All Rights Reserved.
+
+#include "PQVectorIndex.h"
+
+THIRD_PARTY_INCLUDES_START
+#include "faiss/index_io.h"
+#include "faiss/impl/io.h"
+THIRD_PARTY_INCLUDES_END
+
+namespace UE::SemanticSearch
+{
+
+FPQVectorIndex::FPQVectorIndex(int32 InDimension, int32 InM, int32 InNBits)
+	: Dimension(InDimension)
+	, M(InM)
+	, NBits(InNBits)
+{
+	check(Dimension % M == 0);
+	faiss::IndexPQ* PQIndex = new faiss::IndexPQ(Dimension, M, NBits);
+	IndexWithIDs = new faiss::IndexIDMap(PQIndex);
+}
+
+FPQVectorIndex::FPQVectorIndex(int32 InDimension, int32 InM, int32 InNBits, faiss::IndexIDMap* InIndex)
+	: Dimension(InDimension)
+	, M(InM)
+	, NBits(InNBits)
+	, IndexWithIDs(InIndex)
+{
+}
+
+FPQVectorIndex::~FPQVectorIndex()
+{
+	delete IndexWithIDs;
+}
+
+void FPQVectorIndex::Train(TConstArrayView<float> Vectors, int64 NumVectors)
+{
+	check(Vectors.Num() == NumVectors * Dimension);
+	IndexWithIDs->train(NumVectors, Vectors.GetData());
+}
+
+bool FPQVectorIndex::IsTrained() const
+{
+	return IndexWithIDs->is_trained;
+}
+
+void FPQVectorIndex::Add(int64 ID, TConstArrayView<float> Vector)
+{
+	check(IsTrained());
+	check(Vector.Num() == Dimension);
+	faiss::idx_t FaissID = static_cast<faiss::idx_t>(ID);
+	IndexWithIDs->add_with_ids(1, Vector.GetData(), &FaissID);
+}
+
+void FPQVectorIndex::AddVectors(TConstArrayView<int64> IDs, TConstArrayView<float> Vectors)
+{
+	check(IsTrained());
+	const int64 Count = IDs.Num();
+	check(Vectors.Num() == Count * Dimension);
+	static_assert(sizeof(faiss::idx_t) == sizeof(int64), "FAISS idx_t size mismatch");
+	IndexWithIDs->add_with_ids(
+		Count,
+		Vectors.GetData(),
+		reinterpret_cast<const faiss::idx_t*>(IDs.GetData()));
+}
+
+bool FPQVectorIndex::Remove(int64 ID)
+{
+	faiss::idx_t FaissID = static_cast<faiss::idx_t>(ID);
+	faiss::IDSelectorBatch Selector(1, &FaissID);
+	return IndexWithIDs->remove_ids(Selector) > 0;
+}
+
+TArray<FSearchResult> FPQVectorIndex::Search(TConstArrayView<float> QueryVector, int32 K) const
+{
+	check(QueryVector.Num() == Dimension);
+
+	TArray<float> Distances;
+	TArray<faiss::idx_t> Labels;
+	Distances.SetNumUninitialized(K);
+	Labels.SetNumUninitialized(K);
+
+	IndexWithIDs->search(1, QueryVector.GetData(), K, Distances.GetData(), Labels.GetData());
+
+	TArray<FSearchResult> Results;
+	Results.Reserve(K);
+	for (int32 i = 0; i < K; ++i)
+	{
+		if (Labels[i] >= 0)
+		{
+			FSearchResult Result;
+			Result.ID = static_cast<int64>(Labels[i]);
+			Result.Distance = Distances[i];
+			Results.Add(Result);
+		}
+	}
+	return Results;
+}
+
+TArray<uint8> FPQVectorIndex::Quantize(TConstArrayView<float> Vector) const
+{
+	check(IsTrained());
+	check(Vector.Num() == Dimension);
+
+	faiss::IndexPQ* PQIndex = static_cast<faiss::IndexPQ*>(IndexWithIDs->index);
+	const int64 CodeSize = PQIndex->pq.code_size;
+
+	TArray<uint8> Codes;
+	Codes.SetNumUninitialized(CodeSize);
+	PQIndex->pq.compute_codes(Vector.GetData(), Codes.GetData(), 1);
+	return Codes;
+}
+
+TArray<uint8> FPQVectorIndex::QuantizeBatch(TConstArrayView<float> Vectors, int64 NumVectors) const
+{
+	check(IsTrained());
+	check(Vectors.Num() == NumVectors * Dimension);
+
+	faiss::IndexPQ* PQIndex = static_cast<faiss::IndexPQ*>(IndexWithIDs->index);
+	const int64 CodeSize = PQIndex->pq.code_size;
+
+	TArray<uint8> Codes;
+	Codes.SetNumUninitialized(NumVectors * CodeSize);
+	PQIndex->pq.compute_codes(Vectors.GetData(), Codes.GetData(), NumVectors);
+	return Codes;
+}
+
+int64 FPQVectorIndex::GetCount() const
+{
+	return IndexWithIDs->ntotal;
+}
+
+TArray<uint8> FPQVectorIndex::Serialize() const
+{
+	faiss::VectorIOWriter Writer;
+	faiss::write_index(IndexWithIDs, &Writer);
+
+	TArray<uint8> Result;
+	Result.Append(
+		reinterpret_cast<const uint8*>(Writer.data.data()),
+		static_cast<int32>(Writer.data.size()));
+	return Result;
+}
+
+TArray<uint8> FPQVectorIndex::SerializeCodebook() const
+{
+	check(IsTrained());
+	faiss::IndexPQ* PQIndex = static_cast<faiss::IndexPQ*>(IndexWithIDs->index);
+
+	faiss::VectorIOWriter Writer;
+	faiss::write_ProductQuantizer(&PQIndex->pq, &Writer);
+
+	TArray<uint8> Result;
+	Result.Append(
+		reinterpret_cast<const uint8*>(Writer.data.data()),
+		static_cast<int32>(Writer.data.size()));
+	return Result;
+}
+
+TUniquePtr<FPQVectorIndex> FPQVectorIndex::DeserializeCodebook(TConstArrayView<uint8> Data, int32 ExpectedDimension)
+{
+	faiss::VectorIOReader Reader;
+	Reader.data.assign(Data.GetData(), Data.GetData() + Data.Num());
+
+	faiss::ProductQuantizer* PQ = faiss::read_ProductQuantizer(&Reader);
+	if (!PQ || static_cast<int32>(PQ->d) != ExpectedDimension)
+	{
+		delete PQ;
+		return nullptr;
+	}
+
+	const int32 LoadedM = static_cast<int32>(PQ->M);
+	const int32 LoadedNBits = static_cast<int32>(PQ->nbits);
+
+	faiss::IndexPQ* PQIndex = new faiss::IndexPQ(ExpectedDimension, LoadedM, LoadedNBits);
+	PQIndex->pq = std::move(*PQ);
+	PQIndex->is_trained = true;
+	delete PQ;
+
+	faiss::IndexIDMap* IDMapIndex = new faiss::IndexIDMap(PQIndex);
+	return TUniquePtr<FPQVectorIndex>(new FPQVectorIndex(ExpectedDimension, LoadedM, LoadedNBits, IDMapIndex));
+}
+
+TUniquePtr<FPQVectorIndex> FPQVectorIndex::Deserialize(TConstArrayView<uint8> Data, int32 ExpectedDimension)
+{
+	faiss::VectorIOReader Reader;
+	Reader.data.assign(Data.GetData(), Data.GetData() + Data.Num());
+
+	faiss::Index* LoadedIndex = faiss::read_index(&Reader);
+	if (!LoadedIndex || LoadedIndex->d != ExpectedDimension)
+	{
+		delete LoadedIndex;
+		return nullptr;
+	}
+
+	faiss::IndexIDMap* IDMapIndex = static_cast<faiss::IndexIDMap*>(LoadedIndex);
+	faiss::IndexPQ* PQIndex = static_cast<faiss::IndexPQ*>(IDMapIndex->index);
+	const int32 LoadedM = static_cast<int32>(PQIndex->pq.M);
+	const int32 LoadedNBits = static_cast<int32>(PQIndex->pq.nbits);
+
+	return TUniquePtr<FPQVectorIndex>(new FPQVectorIndex(ExpectedDimension, LoadedM, LoadedNBits, IDMapIndex));
+}
+
+} // namespace UE::SemanticSearch
diff --git a/Engine/Plugins/Experimental/SemanticSearch/Source/SemanticSearch/Private/VectorIndex/PQVectorIndex.h b/Engine/Plugins/Experimental/SemanticSearch/Source/SemanticSearch/Private/VectorIndex/PQVectorIndex.h
new file mode 100644
--- /dev/null
+++ b/Engine/Plugins/Experimental/SemanticSearch/Source/SemanticSearch/Private/VectorIndex/PQVectorIndex.h
@@ -0,0 +1,70 @@
+// Copyright Epic Games, Inc. All Rights Reserved.
+
+#pragma once
+
+#include "Interfaces/IVectorIndex.h"
+#include "Templates/UniquePtr.h"
+
+THIRD_PARTY_INCLUDES_START
+#include "faiss/IndexPQ.h"
+#include "faiss/IndexIDMap.h"
+THIRD_PARTY_INCLUDES_END
+
+namespace UE::SemanticSearch
+{
+
+/**
+ * Product Quantization vector index.
+ *
+ * Compresses each vector into M bytes (with nbits=8), giving ~32x compression
+ * for 768-dim vectors (3072 bytes -> 96 bytes). Must be trained before use.
+ */
+class FPQVectorIndex : public IVectorIndex
+{
+public:
+	FPQVectorIndex(int32 InDimension, int32 InM = 96, int32 InNBits = 8);
+	virtual ~FPQVectorIndex() override;
+
+	FPQVectorIndex(const FPQVectorIndex&) = delete;
+	FPQVectorIndex& operator=(const FPQVectorIndex&) = delete;
+
+	// IVectorIndex
+	virtual void Train(TConstArrayView<float> Vectors, int64 NumVectors) override;
+	virtual bool IsTrained() const override;
+	virtual void Add(int64 ID, TConstArrayView<float> Vector) override;
+	virtual bool Remove(int64 ID) override;
+	virtual TArray<FSearchResult> Search(TConstArrayView<float> QueryVector, int32 K) const override;
+	virtual int64 GetCount() const override;
+	virtual int32 GetDimension() const override { return Dimension; }
+	virtual TArray<uint8> Serialize() const override;
+	virtual TArray<uint8> Quantize(TConstArrayView<float> Vector) const override;
+	virtual bool SupportsQuantization() const override { return true; }
+
+	/** Batch add multiple vectors with corresponding IDs. */
+	void AddVectors(TConstArrayView<int64> IDs, TConstArrayView<float> Vectors);
+
+	/** Batch quantize vectors without adding them. Returns NumVectors * M bytes. */
+	TArray<uint8> QuantizeBatch(TConstArrayView<float> Vectors, int64 NumVectors) const;
+
+	int32 GetM() const { return M; }
+	int32 GetNBits() const { return NBits; }
+
+	/** Serialize just the PQ codebook (centroids, M, nbits). */
+	TArray<uint8> SerializeCodebook() const;
+
+	/** Create a trained empty index from a serialized codebook. */
+	static TUniquePtr<FPQVectorIndex> DeserializeCodebook(TConstArrayView<uint8> Data, int32 ExpectedDimension);
+
+	/** Deserialize a full index. */
+	static TUniquePtr<FPQVectorIndex> Deserialize(TConstArrayView<uint8> Data, int32 ExpectedDimension);
+
+private:
+	FPQVectorIndex(int32 InDimension, int32 InM, int32 InNBits, faiss::IndexIDMap* InIndex);
+
+	int32 Dimension;
+	int32 M;
+	int32 NBits;
+	faiss::IndexIDMap* IndexWithIDs = nullptr;
+};
+
+} // namespace UE::SemanticSearch
diff --git a/Engine/Plugins/Experimental/SemanticSearch/Source/SemanticSearch/Public/AssetProcessorManager.h b/Engine/Plugins/Experimental/SemanticSearch/Source/SemanticSearch/Public/AssetProcessorManager.h
--- a/Engine/Plugins/Experimental/SemanticSearch/Source/SemanticSearch/Public/AssetProcessorManager.h
+++ b/Engine/Plugins/Experimental/SemanticSearch/Source/SemanticSearch/Public/AssetProcessorManager.h
@@ -10,6 +10,7 @@
 #include "DerivedDataCacheKey.h"
 #include "HAL/CriticalSection.h"
 #include "Interfaces/IAssetProcessor.h"
+#include "Interfaces/IVectorIndex.h"
 #include "Templates/Function.h"
 #include "Templates/UniquePtr.h"
 #include "Templates/UnrealTypeTraits.h"
@@ -44,6 +45,9 @@ public:
 	SEMANTICSEARCH_API bool GetCaptionData(const FAssetData& InAssetData, TUniqueFunction<void(FAssetCaptionResult&&, FString&&)>&& InOnDataAvailable);
 	SEMANTICSEARCH_API bool GetEmbeddingData(const FAssetData& InAssetData, TUniqueFunction<void(FAssetEmbeddingResult&&, FString&&)>&& InOnDataAvailable);
 
+	SEMANTICSEARCH_API void SetVectorIndex(TSharedPtr<IVectorIndex> InIndex);
+	SEMANTICSEARCH_API TSharedPtr<IVectorIndex> GetVectorIndex() const;
+
 	template<typename FAssetProcessor>
 	void RegisterAssetProcessor();
 
@@ -75,6 +79,7 @@ private:
 
 	TStripedMap<32, DerivedData::FCacheKey, TSharedRef<Private::FSemanticSearchBuildCacheTask>> BuildsJobs;
 
+	TSharedPtr<IVectorIndex> VectorIndex;
 };
 
 template<typename FAssetProcessor>
diff --git a/Engine/Plugins/Experimental/SemanticSearch/Source/SemanticSearch/Public/ISemanticSearchModule.h b/Engine/Plugins/Experimental/SemanticSearch/Source/SemanticSearch/Public/ISemanticSearchModule.h
--- a/Engine/Plugins/Experimental/SemanticSearch/Source/SemanticSearch/Public/ISemanticSearchModule.h
+++ b/Engine/Plugins/Experimental/SemanticSearch/Source/SemanticSearch/Public/ISemanticSearchModule.h
@@ -3,6 +3,7 @@
 #pragma once
 
 #include "Interfaces/IEmbeddingProvider.h"
+#include "Interfaces/IVectorIndex.h"
 #include "Modules/ModuleInterface.h"
 #include "Modules/ModuleManager.h"
 #include "Templates/SharedPointer.h"
@@ -20,5 +21,8 @@ namespace UE::SemanticSearch
 		}
 
 		virtual void RegisterEmbeddingProvider(TUniquePtr<IEmbeddingProvider>&& EmbeddingProvider) = 0;
+
+		/** Set the vector index used for similarity search. Replaces any existing index. */
+		virtual void SetVectorIndex(TSharedPtr<IVectorIndex> InIndex) = 0;
 	};
 }
\ No newline at end of file
diff --git a/Engine/Plugins/Experimental/SemanticSearch/Source/SemanticSearch/Public/Interfaces/IVectorIndex.h b/Engine/Plugins/Experimental/SemanticSearch/Source/SemanticSearch/Public/Interfaces/IVectorIndex.h
new file mode 100644
--- /dev/null
+++ b/Engine/Plugins/Experimental/SemanticSearch/Source/SemanticSearch/Public/Interfaces/IVectorIndex.h
@@ -0,0 +1,60 @@
+// Copyright Epic Games, Inc. All Rights Reserved.
+
+#pragma once
+
+#include "Containers/Array.h"
+#include "Containers/ArrayView.h"
+
+namespace UE::SemanticSearch
+{
+
+struct FSearchResult
+{
+	int64 ID = -1;
+	float Distance = 0.0f;
+};
+
+/**
+ * Abstract interface for vector similarity search indices.
+ *
+ * Concrete implementations include:
+ *   - FFlatVectorIndex: Exact brute-force search (IndexFlatL2). No training needed.
+ *   - FPQVectorIndex: Product Quantization compressed search. Requires training.
+ */
+class IVectorIndex
+{
+public:
+	virtual ~IVectorIndex() = default;
+
+	/** Train the index on representative vectors. No-op for flat indices. */
+	virtual void Train(TConstArrayView<float> Vectors, int64 NumVectors) = 0;
+
+	/** Whether the index has been trained and is ready for Add/Search. */
+	virtual bool IsTrained() const = 0;
+
+	/** Add a single vector with an external ID. Index must be trained first. */
+	virtual void Add(int64 ID, TConstArrayView<float> Vector) = 0;
+
+	/** Remove a vector by ID. Returns true if found and removed. */
+	virtual bool Remove(int64 ID) = 0;
+
+	/** Search for K nearest neighbors. Returns results sorted by ascending distance. */
+	virtual TArray<FSearchResult> Search(TConstArrayView<float> QueryVector, int32 K) const = 0;
+
+	/** Number of vectors currently in the index. */
+	virtual int64 GetCount() const = 0;
+
+	/** Embedding dimension. */
+	virtual int32 GetDimension() const = 0;
+
+	/** Serialize the full index to a byte array. */
+	virtual TArray<uint8> Serialize() const = 0;
+
+	/** Quantize a vector into compressed form (e.g. PQ codes). Returns empty if not supported. */
+	virtual TArray<uint8> Quantize(TConstArrayView<float> Vector) const { return {}; }
+
+	/** Whether this index supports quantizing vectors into a compressed representation. */
+	virtual bool SupportsQuantization() const { return false; }
+};
+
+} // namespace UE::SemanticSearch
diff --git a/Engine/Plugins/Experimental/SemanticSearch/Source/SemanticSearch/SemanticSearch.Build.cs b/Engine/Plugins/Experimental/SemanticSearch/Source/SemanticSearch/SemanticSearch.Build.cs
--- a/Engine/Plugins/Experimental/SemanticSearch/Source/SemanticSearch/SemanticSearch.Build.cs
+++ b/Engine/Plugins/Experimental/SemanticSearch/Source/SemanticSearch/SemanticSearch.Build.cs
@@ -25,6 +25,7 @@ public class SemanticSearch : ModuleRules
 				"Json",
 				"PhysicsCore",
 				"UnrealEd",
+				"FAISS"
 			});
 
 	}
