[Backout] - CL52103122
#rnx
Original CL Desc
-----------------------------------------------------------------
Abstract vector index interface + FAISS integration into DDC pipeline

- IVectorIndex: abstract interface with Train, Add, Remove, Search, Encode
- FPQVectorIndex: Product Quantization compressed search
- FFlatVectorIndex: Brute-force L2 search, no training needed
- DDC builder encodes embeddings (PQ codes) and adds vectors to index
- AssetProcessorManager owns the index
- GetAssetSemanticID: stable int64 hash for asset identification

#Synced-CL 52058676

#Synced-CL 52058676

diff --git a/Engine/Plugins/Experimental/SemanticSearch/Source/SemanticSearch/Private/AssetProcessorManager.cpp b/Engine/Plugins/Experimental/SemanticSearch/Source/SemanticSearch/Private/AssetProcessorManager.cpp
--- a/Engine/Plugins/Experimental/SemanticSearch/Source/SemanticSearch/Private/AssetProcessorManager.cpp
+++ b/Engine/Plugins/Experimental/SemanticSearch/Source/SemanticSearch/Private/AssetProcessorManager.cpp
@@ -100,15 +100,7 @@ namespace UE::SemanticSearch
 		});
 	}
 
-	void FAssetProcessorManager::SetVectorIndex(TSharedPtr<IVectorIndex> InIndex)
-	{
-		VectorIndex = MoveTemp(InIndex);
-	}
-
-	TSharedPtr<IVectorIndex> FAssetProcessorManager::GetVectorIndex() const
-	{
-		return VectorIndex;
-	}
+	
 
 	void FAssetProcessorManager::RegisterAssetProcessor(TSharedRef<IAssetProcessor>&& AssetProcessor)
 	{
diff --git a/Engine/Plugins/Experimental/SemanticSearch/Source/SemanticSearch/Private/AssetProcessors/AssetProcessorUtils.cpp b/Engine/Plugins/Experimental/SemanticSearch/Source/SemanticSearch/Private/AssetProcessors/AssetProcessorUtils.cpp
--- a/Engine/Plugins/Experimental/SemanticSearch/Source/SemanticSearch/Private/AssetProcessors/AssetProcessorUtils.cpp
+++ b/Engine/Plugins/Experimental/SemanticSearch/Source/SemanticSearch/Private/AssetProcessors/AssetProcessorUtils.cpp
@@ -68,17 +68,4 @@ void RegisterDefaultAssetProcessors()
 	Manager.RegisterAssetProcessor<FSkeletalMeshProcessor>();
 	Manager.RegisterAssetProcessor<FTextureProcessor>();
 }
-
-int64 GetAssetIndexID(const FAssetData& Asset)
-{
-	const FString Path = Asset.GetObjectPathString();
-	FBlake3 Hasher;
-	Hasher.Update(Path.GetCharArray().GetData(), Path.Len() * sizeof(TCHAR));
-	FBlake3Hash Hash = Hasher.Finalize();
-	int64 Result;
-	FMemory::Memcpy(&Result, Hash.GetBytes(), sizeof(int64));
-	// Ensure non-negative (FAISS uses -1 as sentinel)
-	return Result & 0x7FFFFFFFFFFFFFFF;
-}
-
 }
diff --git a/Engine/Plugins/Experimental/SemanticSearch/Source/SemanticSearch/Private/AssetProcessors/AssetProcessorUtils.h b/Engine/Plugins/Experimental/SemanticSearch/Source/SemanticSearch/Private/AssetProcessors/AssetProcessorUtils.h
--- a/Engine/Plugins/Experimental/SemanticSearch/Source/SemanticSearch/Private/AssetProcessors/AssetProcessorUtils.h
+++ b/Engine/Plugins/Experimental/SemanticSearch/Source/SemanticSearch/Private/AssetProcessors/AssetProcessorUtils.h
@@ -56,10 +56,4 @@ void PopulateFromEnum(TMap<FString, FString>& OutMap, const UEnum* Enum);
  */
 void RegisterDefaultAssetProcessors();
 
-/**
- * Compute a stable int64 hash of an asset's object path, suitable for use as a vector index ID.
- * Uses Blake3. Always non-negative (FAISS reserves -1 as sentinel).
- */
-int64 GetAssetIndexID(const FAssetData& Asset);
-
 }
diff --git a/Engine/Plugins/Experimental/SemanticSearch/Source/SemanticSearch/Private/SemanticSearchDerivedDataBuilder.cpp b/Engine/Plugins/Experimental/SemanticSearch/Source/SemanticSearch/Private/SemanticSearchDerivedDataBuilder.cpp
--- a/Engine/Plugins/Experimental/SemanticSearch/Source/SemanticSearch/Private/SemanticSearchDerivedDataBuilder.cpp
+++ b/Engine/Plugins/Experimental/SemanticSearch/Source/SemanticSearch/Private/SemanticSearchDerivedDataBuilder.cpp
@@ -2,14 +2,11 @@
 
 #include "SemanticSearchDerivedDataBuilder.h"
 
-#include "AssetProcessorManager.h"
-#include "AssetProcessors/AssetProcessorUtils.h"
 #include "Containers/SharedString.h"
 #include "ProfilingDebugging/CpuProfilerTrace.h"
 #include "Containers/StringFwd.h"
 #include "DerivedDataCache.h"
 #include "Interfaces/IAssetProcessor.h"
-#include "Interfaces/IVectorIndex.h"
 #include "Misc/Guid.h"
 #include "Misc/ScopeRWLock.h"
 #include "SemanticSearchModule.h"
@@ -32,21 +29,11 @@ DerivedData::FValueId GetCaptionDatachunkId()
 	return DerivedData::FValueId::FromName(TEXT("CaptionData"));
 }
 
-DerivedData::FValueId GetQuantizedEmbeddingDatachunkId()
-{
-	return DerivedData::FValueId::FromName(TEXT("QuantizedEmbeddingData"));
-}
-
 void SerializeEmbedding(FEmbeddingDerivedData& Embedding, FArchive& Archive)
 {
 	Archive << Embedding.Embedding;
 }
 
-void SerializeQuantizedEmbedding(FQuantizedEmbeddingDerivedData& QuantizedEmbedding, FArchive& Archive)
-{
-	Archive << QuantizedEmbedding.QuantizedEmbedding;
-}
-
 void SerializeCaption(FCaptionDerivedData& Caption, FArchive& Archive)
 {
 	Archive << Caption.Caption;
@@ -272,16 +259,6 @@ void FSemanticSearchBuildCacheTask::BuildDerivedData()
 											FEmbeddingDerivedData EmbeddingDerivedData;
 											EmbeddingDerivedData.Embedding = Embedding.Embedding;
 
-											// Quantize the embedding if the active index supports it (PQ).
-											FQuantizedEmbeddingDerivedData QuantizedEmbeddingData;
-											if (TSharedPtr<IVectorIndex> VectorIndex = FAssetProcessorManager::Get().GetVectorIndex())
-											{
-												if (VectorIndex->SupportsQuantization() && VectorIndex->IsTrained())
-												{
-													QuantizedEmbeddingData.QuantizedEmbedding = VectorIndex->Quantize(EmbeddingDerivedData.Embedding);
-												}
-											}
-
 											if (AssetKey == GenerateCacheKey(*Asset, AssetProcessor))
 											{
 												using namespace DerivedData;
@@ -304,15 +281,6 @@ void FSemanticSearchBuildCacheTask::BuildDerivedData()
 													RecordBuilder.AddValue(GetEmbeddingDatachunkId(), MakeSharedBufferFromArray(MoveTemp(Buffer)));
 												}
 
-												// Save quantized embedding (PQ codes) if available
-												if (QuantizedEmbeddingData.QuantizedEmbedding.Num() > 0)
-												{
-													TArray<uint8> Buffer;
-													FMemoryWriter Writer(Buffer);
-													SerializeQuantizedEmbedding(QuantizedEmbeddingData, Writer);
-													RecordBuilder.AddValue(GetQuantizedEmbeddingDatachunkId(), MakeSharedBufferFromArray(MoveTemp(Buffer)));
-												}
-
 												DerivedData::FCachePutRequest PutRequest{ DDCRequestName, RecordBuilder.Build(), GetCachePolicy() | ECachePolicy::SkipData };
 
 												DerivedData::GetCache().Put(
@@ -331,13 +299,6 @@ void FSemanticSearchBuildCacheTask::BuildDerivedData()
 														}
 													});
 
-												// Add the embedding to the vector index
-												if (TSharedPtr<IVectorIndex> VectorIndex = FAssetProcessorManager::Get().GetVectorIndex())
-												{
-													const int64 AssetID = GetAssetIndexID(*Asset);
-													VectorIndex->Add(AssetID, EmbeddingDerivedData.Embedding);
-												}
-
 												BroadcastQueuedRequests(MoveTemp(EmbeddingDerivedData), MoveTemp(CaptionDerivedData), FString());
 											}
 											else
diff --git a/Engine/Plugins/Experimental/SemanticSearch/Source/SemanticSearch/Private/SemanticSearchDerivedDataBuilder.h b/Engine/Plugins/Experimental/SemanticSearch/Source/SemanticSearch/Private/SemanticSearchDerivedDataBuilder.h
--- a/Engine/Plugins/Experimental/SemanticSearch/Source/SemanticSearch/Private/SemanticSearchDerivedDataBuilder.h
+++ b/Engine/Plugins/Experimental/SemanticSearch/Source/SemanticSearch/Private/SemanticSearchDerivedDataBuilder.h
@@ -29,11 +29,6 @@ struct FEmbeddingDerivedData
 	TArray<float> Embedding;
 };
 
-struct FQuantizedEmbeddingDerivedData
-{
-	TArray<uint8> QuantizedEmbedding;
-};
-
 struct FCaptionDerivedData
 {
 	FString Caption;
diff --git a/Engine/Plugins/Experimental/SemanticSearch/Source/SemanticSearch/Private/SemanticSearchModule.cpp b/Engine/Plugins/Experimental/SemanticSearch/Source/SemanticSearch/Private/SemanticSearchModule.cpp
--- a/Engine/Plugins/Experimental/SemanticSearch/Source/SemanticSearch/Private/SemanticSearchModule.cpp
+++ b/Engine/Plugins/Experimental/SemanticSearch/Source/SemanticSearch/Private/SemanticSearchModule.cpp
@@ -2,9 +2,7 @@
 
 #include "SemanticSearchModule.h"
 
-#include "AssetProcessorManager.h"
 #include "AssetProcessors/AssetProcessorUtils.h"
-#include "VectorIndex/FlatVectorIndex.h"
 
 #include "Modules/ModuleManager.h"
 #include "RemoteEmbeddingProvider.h"
@@ -26,11 +24,6 @@ void FSemanticSearchModule::StartupModule()
 	UE_LOG(LogSemanticSearch, Log, TEXT("SemanticSearch module starting up"));
 	RegisteredEmbeddingProvider = MakeShared<FRemoteEmbeddingProvider>();
 	RegisterDefaultAssetProcessors();
-
-	// Default to a flat index. Dimension is set when first embedding arrives.
-	// Callers can replace this with a PQ index via SetVectorIndex().
-	constexpr int32 DefaultDimension = 768;
-	FAssetProcessorManager::Get().SetVectorIndex(MakeShared<FFlatVectorIndex>(DefaultDimension));
 }
 
 void FSemanticSearchModule::ShutdownModule()
@@ -48,10 +41,5 @@ void FSemanticSearchModule::RegisterEmbeddingProvider(TUniquePtr<IEmbeddingProvi
 	RegisteredEmbeddingProvider = MakeShareable<IEmbeddingProvider>(EmbeddingProvider.Release());
 }
 
-void FSemanticSearchModule::SetVectorIndex(TSharedPtr<IVectorIndex> InIndex)
-{
-	FAssetProcessorManager::Get().SetVectorIndex(MoveTemp(InIndex));
-}
-
 }
 
diff --git a/Engine/Plugins/Experimental/SemanticSearch/Source/SemanticSearch/Private/SemanticSearchModule.h b/Engine/Plugins/Experimental/SemanticSearch/Source/SemanticSearch/Private/SemanticSearchModule.h
--- a/Engine/Plugins/Experimental/SemanticSearch/Source/SemanticSearch/Private/SemanticSearchModule.h
+++ b/Engine/Plugins/Experimental/SemanticSearch/Source/SemanticSearch/Private/SemanticSearchModule.h
@@ -24,7 +24,6 @@ public:
 
 
 	virtual void RegisterEmbeddingProvider(TUniquePtr<IEmbeddingProvider>&& EmbeddingProvider) override;
-	virtual void SetVectorIndex(TSharedPtr<IVectorIndex> InIndex) override;
 
 private:
 	TSharedPtr<IEmbeddingProvider> RegisteredEmbeddingProvider;
diff --git a/Engine/Plugins/Experimental/SemanticSearch/Source/SemanticSearch/Private/VectorIndex/FlatVectorIndex.cpp b/Engine/Plugins/Experimental/SemanticSearch/Source/SemanticSearch/Private/VectorIndex/FlatVectorIndex.cpp
deleted file mode 100644
diff --git a/Engine/Plugins/Experimental/SemanticSearch/Source/SemanticSearch/Private/VectorIndex/FlatVectorIndex.h b/Engine/Plugins/Experimental/SemanticSearch/Source/SemanticSearch/Private/VectorIndex/FlatVectorIndex.h
deleted file mode 100644
diff --git a/Engine/Plugins/Experimental/SemanticSearch/Source/SemanticSearch/Private/VectorIndex/PQVectorIndex.cpp b/Engine/Plugins/Experimental/SemanticSearch/Source/SemanticSearch/Private/VectorIndex/PQVectorIndex.cpp
deleted file mode 100644
diff --git a/Engine/Plugins/Experimental/SemanticSearch/Source/SemanticSearch/Private/VectorIndex/PQVectorIndex.h b/Engine/Plugins/Experimental/SemanticSearch/Source/SemanticSearch/Private/VectorIndex/PQVectorIndex.h
deleted file mode 100644
diff --git a/Engine/Plugins/Experimental/SemanticSearch/Source/SemanticSearch/Public/AssetProcessorManager.h b/Engine/Plugins/Experimental/SemanticSearch/Source/SemanticSearch/Public/AssetProcessorManager.h
--- a/Engine/Plugins/Experimental/SemanticSearch/Source/SemanticSearch/Public/AssetProcessorManager.h
+++ b/Engine/Plugins/Experimental/SemanticSearch/Source/SemanticSearch/Public/AssetProcessorManager.h
@@ -10,7 +10,6 @@
 #include "DerivedDataCacheKey.h"
 #include "HAL/CriticalSection.h"
 #include "Interfaces/IAssetProcessor.h"
-#include "Interfaces/IVectorIndex.h"
 #include "Templates/Function.h"
 #include "Templates/UniquePtr.h"
 #include "Templates/UnrealTypeTraits.h"
@@ -45,9 +44,6 @@ public:
 	SEMANTICSEARCH_API bool GetCaptionData(const FAssetData& InAssetData, TUniqueFunction<void(FAssetCaptionResult&&, FString&&)>&& InOnDataAvailable);
 	SEMANTICSEARCH_API bool GetEmbeddingData(const FAssetData& InAssetData, TUniqueFunction<void(FAssetEmbeddingResult&&, FString&&)>&& InOnDataAvailable);
 
-	SEMANTICSEARCH_API void SetVectorIndex(TSharedPtr<IVectorIndex> InIndex);
-	SEMANTICSEARCH_API TSharedPtr<IVectorIndex> GetVectorIndex() const;
-
 	template<typename FAssetProcessor>
 	void RegisterAssetProcessor();
 
@@ -79,7 +75,6 @@ private:
 
 	TStripedMap<32, DerivedData::FCacheKey, TSharedRef<Private::FSemanticSearchBuildCacheTask>> BuildsJobs;
 
-	TSharedPtr<IVectorIndex> VectorIndex;
 };
 
 template<typename FAssetProcessor>
diff --git a/Engine/Plugins/Experimental/SemanticSearch/Source/SemanticSearch/Public/ISemanticSearchModule.h b/Engine/Plugins/Experimental/SemanticSearch/Source/SemanticSearch/Public/ISemanticSearchModule.h
--- a/Engine/Plugins/Experimental/SemanticSearch/Source/SemanticSearch/Public/ISemanticSearchModule.h
+++ b/Engine/Plugins/Experimental/SemanticSearch/Source/SemanticSearch/Public/ISemanticSearchModule.h
@@ -3,7 +3,6 @@
 #pragma once
 
 #include "Interfaces/IEmbeddingProvider.h"
-#include "Interfaces/IVectorIndex.h"
 #include "Modules/ModuleInterface.h"
 #include "Modules/ModuleManager.h"
 #include "Templates/SharedPointer.h"
@@ -21,8 +20,5 @@ namespace UE::SemanticSearch
 		}
 
 		virtual void RegisterEmbeddingProvider(TUniquePtr<IEmbeddingProvider>&& EmbeddingProvider) = 0;
-
-		/** Set the vector index used for similarity search. Replaces any existing index. */
-		virtual void SetVectorIndex(TSharedPtr<IVectorIndex> InIndex) = 0;
 	};
 }
\ No newline at end of file
diff --git a/Engine/Plugins/Experimental/SemanticSearch/Source/SemanticSearch/Public/Interfaces/IVectorIndex.h b/Engine/Plugins/Experimental/SemanticSearch/Source/SemanticSearch/Public/Interfaces/IVectorIndex.h
deleted file mode 100644
diff --git a/Engine/Plugins/Experimental/SemanticSearch/Source/SemanticSearch/SemanticSearch.Build.cs b/Engine/Plugins/Experimental/SemanticSearch/Source/SemanticSearch/SemanticSearch.Build.cs
--- a/Engine/Plugins/Experimental/SemanticSearch/Source/SemanticSearch/SemanticSearch.Build.cs
+++ b/Engine/Plugins/Experimental/SemanticSearch/Source/SemanticSearch/SemanticSearch.Build.cs
@@ -25,7 +25,6 @@ public class SemanticSearch : ModuleRules
 				"Json",
 				"PhysicsCore",
 				"UnrealEd",
-				"FAISS"
 			});
 
 	}
