Add ID-filtered search to vector index interface

#Synced-CL 52058676

diff --git a/Engine/Plugins/Experimental/SemanticSearch/Source/SemanticSearch/Private/VectorIndex/FlatVectorIndex.cpp b/Engine/Plugins/Experimental/SemanticSearch/Source/SemanticSearch/Private/VectorIndex/FlatVectorIndex.cpp
--- a/Engine/Plugins/Experimental/SemanticSearch/Source/SemanticSearch/Private/VectorIndex/FlatVectorIndex.cpp
+++ b/Engine/Plugins/Experimental/SemanticSearch/Source/SemanticSearch/Private/VectorIndex/FlatVectorIndex.cpp
@@ -48,7 +48,7 @@ bool FFlatVectorIndex::Remove(int64 ID)
 	return IndexWithIDs->remove_ids(Selector) > 0;
 }
 
-TArray<FSearchResult> FFlatVectorIndex::Search(TConstArrayView<float> QueryVector, int32 K) const
+TArray<FSearchResult> FFlatVectorIndex::Search(TConstArrayView<float> QueryVector, int32 K, TConstArrayView<int64> IDFilter) const
 {
 	check(QueryVector.Num() == Dimension);
 
@@ -57,7 +57,18 @@ TArray<FSearchResult> FFlatVectorIndex::Search(TConstArrayView<float> QueryVecto
 	Distances.SetNumUninitialized(K);
 	Labels.SetNumUninitialized(K);
 
-	IndexWithIDs->search(1, QueryVector.GetData(), K, Distances.GetData(), Labels.GetData());
+	if (IDFilter.Num() > 0)
+	{
+		static_assert(sizeof(faiss::idx_t) == sizeof(int64), "FAISS idx_t size mismatch");
+		faiss::IDSelectorBatch Selector(IDFilter.Num(), reinterpret_cast<const faiss::idx_t*>(IDFilter.GetData()));
+		faiss::SearchParameters Params;
+		Params.sel = &Selector;
+		IndexWithIDs->search(1, QueryVector.GetData(), K, Distances.GetData(), Labels.GetData(), &Params);
+	}
+	else
+	{
+		IndexWithIDs->search(1, QueryVector.GetData(), K, Distances.GetData(), Labels.GetData());
+	}
 
 	TArray<FSearchResult> Results;
 	Results.Reserve(K);
diff --git a/Engine/Plugins/Experimental/SemanticSearch/Source/SemanticSearch/Private/VectorIndex/FlatVectorIndex.h b/Engine/Plugins/Experimental/SemanticSearch/Source/SemanticSearch/Private/VectorIndex/FlatVectorIndex.h
--- a/Engine/Plugins/Experimental/SemanticSearch/Source/SemanticSearch/Private/VectorIndex/FlatVectorIndex.h
+++ b/Engine/Plugins/Experimental/SemanticSearch/Source/SemanticSearch/Private/VectorIndex/FlatVectorIndex.h
@@ -33,7 +33,7 @@ public:
 	virtual bool IsTrained() const override { return true; }
 	virtual void Add(int64 ID, TConstArrayView<float> Vector) override;
 	virtual bool Remove(int64 ID) override;
-	virtual TArray<FSearchResult> Search(TConstArrayView<float> QueryVector, int32 K) const override;
+	virtual TArray<FSearchResult> Search(TConstArrayView<float> QueryVector, int32 K, TConstArrayView<int64> IDFilter = {}) const override;
 	virtual int64 GetCount() const override;
 	virtual int32 GetDimension() const override { return Dimension; }
 	virtual TArray<uint8> Serialize() const override;
diff --git a/Engine/Plugins/Experimental/SemanticSearch/Source/SemanticSearch/Private/VectorIndex/PQVectorIndex.cpp b/Engine/Plugins/Experimental/SemanticSearch/Source/SemanticSearch/Private/VectorIndex/PQVectorIndex.cpp
--- a/Engine/Plugins/Experimental/SemanticSearch/Source/SemanticSearch/Private/VectorIndex/PQVectorIndex.cpp
+++ b/Engine/Plugins/Experimental/SemanticSearch/Source/SemanticSearch/Private/VectorIndex/PQVectorIndex.cpp
@@ -71,7 +71,7 @@ bool FPQVectorIndex::Remove(int64 ID)
 	return IndexWithIDs->remove_ids(Selector) > 0;
 }
 
-TArray<FSearchResult> FPQVectorIndex::Search(TConstArrayView<float> QueryVector, int32 K) const
+TArray<FSearchResult> FPQVectorIndex::Search(TConstArrayView<float> QueryVector, int32 K, TConstArrayView<int64> IDFilter) const
 {
 	check(QueryVector.Num() == Dimension);
 
@@ -80,7 +80,18 @@ TArray<FSearchResult> FPQVectorIndex::Search(TConstArrayView<float> QueryVector,
 	Distances.SetNumUninitialized(K);
 	Labels.SetNumUninitialized(K);
 
-	IndexWithIDs->search(1, QueryVector.GetData(), K, Distances.GetData(), Labels.GetData());
+	if (IDFilter.Num() > 0)
+	{
+		static_assert(sizeof(faiss::idx_t) == sizeof(int64), "FAISS idx_t size mismatch");
+		faiss::IDSelectorBatch Selector(IDFilter.Num(), reinterpret_cast<const faiss::idx_t*>(IDFilter.GetData()));
+		faiss::SearchParameters Params;
+		Params.sel = &Selector;
+		IndexWithIDs->search(1, QueryVector.GetData(), K, Distances.GetData(), Labels.GetData(), &Params);
+	}
+	else
+	{
+		IndexWithIDs->search(1, QueryVector.GetData(), K, Distances.GetData(), Labels.GetData());
+	}
 
 	TArray<FSearchResult> Results;
 	Results.Reserve(K);
diff --git a/Engine/Plugins/Experimental/SemanticSearch/Source/SemanticSearch/Private/VectorIndex/PQVectorIndex.h b/Engine/Plugins/Experimental/SemanticSearch/Source/SemanticSearch/Private/VectorIndex/PQVectorIndex.h
--- a/Engine/Plugins/Experimental/SemanticSearch/Source/SemanticSearch/Private/VectorIndex/PQVectorIndex.h
+++ b/Engine/Plugins/Experimental/SemanticSearch/Source/SemanticSearch/Private/VectorIndex/PQVectorIndex.h
@@ -33,7 +33,7 @@ public:
 	virtual bool IsTrained() const override;
 	virtual void Add(int64 ID, TConstArrayView<float> Vector) override;
 	virtual bool Remove(int64 ID) override;
-	virtual TArray<FSearchResult> Search(TConstArrayView<float> QueryVector, int32 K) const override;
+	virtual TArray<FSearchResult> Search(TConstArrayView<float> QueryVector, int32 K, TConstArrayView<int64> IDFilter = {}) const override;
 	virtual int64 GetCount() const override;
 	virtual int32 GetDimension() const override { return Dimension; }
 	virtual TArray<uint8> Serialize() const override;
diff --git a/Engine/Plugins/Experimental/SemanticSearch/Source/SemanticSearch/Public/Interfaces/IVectorIndex.h b/Engine/Plugins/Experimental/SemanticSearch/Source/SemanticSearch/Public/Interfaces/IVectorIndex.h
--- a/Engine/Plugins/Experimental/SemanticSearch/Source/SemanticSearch/Public/Interfaces/IVectorIndex.h
+++ b/Engine/Plugins/Experimental/SemanticSearch/Source/SemanticSearch/Public/Interfaces/IVectorIndex.h
@@ -38,8 +38,11 @@ public:
 	/** Remove a vector by ID. Returns true if found and removed. */
 	virtual bool Remove(int64 ID) = 0;
 
-	/** Search for K nearest neighbors. Returns results sorted by ascending distance. */
-	virtual TArray<FSearchResult> Search(TConstArrayView<float> QueryVector, int32 K) const = 0;
+	/**
+	 * Search for K nearest neighbors. Returns results sorted by ascending distance.
+	 * @param IDFilter  If non-empty, only search among vectors with these IDs.
+	 */
+	virtual TArray<FSearchResult> Search(TConstArrayView<float> QueryVector, int32 K, TConstArrayView<int64> IDFilter = {}) const = 0;
 
 	/** Number of vectors currently in the index. */
 	virtual int64 GetCount() const = 0;
