Simplify Index API, support batched adds/remove, make thread-safe

#Synced-CL 52058676

diff --git a/Engine/Plugins/Experimental/SemanticSearch/Source/SemanticSearch/Private/VectorIndex/FlatVectorIndex.cpp b/Engine/Plugins/Experimental/SemanticSearch/Source/SemanticSearch/Private/VectorIndex/FlatVectorIndex.cpp
--- a/Engine/Plugins/Experimental/SemanticSearch/Source/SemanticSearch/Private/VectorIndex/FlatVectorIndex.cpp
+++ b/Engine/Plugins/Experimental/SemanticSearch/Source/SemanticSearch/Private/VectorIndex/FlatVectorIndex.cpp
@@ -2,6 +2,8 @@
 
 #include "FlatVectorIndex.h"
 
+#include "Misc/ScopeRWLock.h"
+
 THIRD_PARTY_INCLUDES_START
 #include "faiss/index_io.h"
 #include "faiss/impl/io.h"
@@ -31,21 +33,30 @@ FFlatVectorIndex::~FFlatVectorIndex()
 void FFlatVectorIndex::Train(TConstArrayView<float> Vectors, int64 NumVectors)
 {
 	// Flat index does not need training.
-	return;
 }
 
-void FFlatVectorIndex::Add(int64 ID, TConstArrayView<float> Vector)
+void FFlatVectorIndex::Add(TConstArrayView<int64> IDs, TConstArrayView<float> Vectors)
 {
-	check(Vector.Num() == Dimension);
-	faiss::idx_t FaissID = static_cast<faiss::idx_t>(ID);
-	IndexWithIDs->add_with_ids(1, Vector.GetData(), &FaissID);
+	const int64 Count = IDs.Num();
+	check(Vectors.Num() == Count * Dimension);
+	static_assert(sizeof(faiss::idx_t) == sizeof(int64), "FAISS idx_t size mismatch");
+	FWriteScopeLock WriteLock(IndexLock);
+	IndexWithIDs->add_with_ids(
+		Count,
+		Vectors.GetData(),
+		reinterpret_cast<const faiss::idx_t*>(IDs.GetData()));
 }
 
-bool FFlatVectorIndex::Remove(int64 ID)
+void FFlatVectorIndex::Remove(TConstArrayView<int64> IDs)
 {
-	faiss::idx_t FaissID = static_cast<faiss::idx_t>(ID);
-	faiss::IDSelectorBatch Selector(1, &FaissID);
-	return IndexWithIDs->remove_ids(Selector) > 0;
+	if (IDs.Num() == 0)
+	{
+		return;
+	}
+	static_assert(sizeof(faiss::idx_t) == sizeof(int64), "FAISS idx_t size mismatch");
+	faiss::IDSelectorBatch Selector(IDs.Num(), reinterpret_cast<const faiss::idx_t*>(IDs.GetData()));
+	FWriteScopeLock WriteLock(IndexLock);
+	IndexWithIDs->remove_ids(Selector);
 }
 
 TArray<FSearchResult> FFlatVectorIndex::Search(TConstArrayView<float> QueryVector, int32 K, TConstArrayView<int64> IDFilter) const
@@ -57,6 +68,8 @@ TArray<FSearchResult> FFlatVectorIndex::Search(TConstArrayView<float> QueryVecto
 	Distances.SetNumUninitialized(K);
 	Labels.SetNumUninitialized(K);
 
+	FReadScopeLock ReadLock(IndexLock);
+
 	if (IDFilter.Num() > 0)
 	{
 		static_assert(sizeof(faiss::idx_t) == sizeof(int64), "FAISS idx_t size mismatch");
@@ -87,11 +100,13 @@ TArray<FSearchResult> FFlatVectorIndex::Search(TConstArrayView<float> QueryVecto
 
 int64 FFlatVectorIndex::GetCount() const
 {
+	FReadScopeLock ReadLock(IndexLock);
 	return IndexWithIDs->ntotal;
 }
 
 TArray<uint8> FFlatVectorIndex::Serialize() const
 {
+	FReadScopeLock ReadLock(IndexLock);
 	faiss::VectorIOWriter Writer;
 	faiss::write_index(IndexWithIDs, &Writer);
 
diff --git a/Engine/Plugins/Experimental/SemanticSearch/Source/SemanticSearch/Private/VectorIndex/FlatVectorIndex.h b/Engine/Plugins/Experimental/SemanticSearch/Source/SemanticSearch/Private/VectorIndex/FlatVectorIndex.h
--- a/Engine/Plugins/Experimental/SemanticSearch/Source/SemanticSearch/Private/VectorIndex/FlatVectorIndex.h
+++ b/Engine/Plugins/Experimental/SemanticSearch/Source/SemanticSearch/Private/VectorIndex/FlatVectorIndex.h
@@ -3,6 +3,7 @@
 #pragma once
 
 #include "Interfaces/IVectorIndex.h"
+#include "HAL/CriticalSection.h"
 #include "Templates/UniquePtr.h"
 
 THIRD_PARTY_INCLUDES_START
@@ -28,11 +29,11 @@ public:
 	FFlatVectorIndex(const FFlatVectorIndex&) = delete;
 	FFlatVectorIndex& operator=(const FFlatVectorIndex&) = delete;
 
-	// IVectorIndex
+	// IVectorIndex batch overrides
 	virtual void Train(TConstArrayView<float> Vectors, int64 NumVectors) override;
 	virtual bool IsTrained() const override { return true; }
-	virtual void Add(int64 ID, TConstArrayView<float> Vector) override;
-	virtual bool Remove(int64 ID) override;
+	virtual void Add(TConstArrayView<int64> IDs, TConstArrayView<float> Vectors) override;
+	virtual void Remove(TConstArrayView<int64> IDs) override;
 	virtual TArray<FSearchResult> Search(TConstArrayView<float> QueryVector, int32 K, TConstArrayView<int64> IDFilter = {}) const override;
 	virtual int64 GetCount() const override;
 	virtual int32 GetDimension() const override { return Dimension; }
@@ -46,6 +47,7 @@ private:
 
 	int32 Dimension;
 	faiss::IndexIDMap* IndexWithIDs = nullptr;
+	mutable FRWLock IndexLock;
 };
 
 } // namespace UE::SemanticSearch
diff --git a/Engine/Plugins/Experimental/SemanticSearch/Source/SemanticSearch/Private/VectorIndex/PQVectorIndex.cpp b/Engine/Plugins/Experimental/SemanticSearch/Source/SemanticSearch/Private/VectorIndex/PQVectorIndex.cpp
--- a/Engine/Plugins/Experimental/SemanticSearch/Source/SemanticSearch/Private/VectorIndex/PQVectorIndex.cpp
+++ b/Engine/Plugins/Experimental/SemanticSearch/Source/SemanticSearch/Private/VectorIndex/PQVectorIndex.cpp
@@ -2,6 +2,8 @@
 
 #include "PQVectorIndex.h"
 
+#include "Misc/ScopeRWLock.h"
+
 THIRD_PARTY_INCLUDES_START
 #include "faiss/index_io.h"
 #include "faiss/impl/io.h"
@@ -10,19 +12,21 @@ THIRD_PARTY_INCLUDES_END
 namespace UE::SemanticSearch
 {
 
-FPQVectorIndex::FPQVectorIndex(int32 InDimension, int32 InM, int32 InNBits)
+FPQVectorIndex::FPQVectorIndex(int32 InDimension, int32 InSubvectorSize, int32 InNBits)
 	: Dimension(InDimension)
-	, M(InM)
+	, SubvectorSize(InSubvectorSize)
+	, NumSubquantizers(InDimension / InSubvectorSize)
 	, NBits(InNBits)
 {
-	check(Dimension % M == 0);
-	faiss::IndexPQ* PQIndex = new faiss::IndexPQ(Dimension, M, NBits);
+	check(Dimension % SubvectorSize == 0);
+	faiss::IndexPQ* PQIndex = new faiss::IndexPQ(Dimension, NumSubquantizers, NBits);
 	IndexWithIDs = new faiss::IndexIDMap(PQIndex);
 }
 
-FPQVectorIndex::FPQVectorIndex(int32 InDimension, int32 InM, int32 InNBits, faiss::IndexIDMap* InIndex)
+FPQVectorIndex::FPQVectorIndex(int32 InDimension, int32 InSubvectorSize, int32 InNBits, faiss::IndexIDMap* InIndex)
 	: Dimension(InDimension)
-	, M(InM)
+	, SubvectorSize(InSubvectorSize)
+	, NumSubquantizers(InDimension / InSubvectorSize)
 	, NBits(InNBits)
 	, IndexWithIDs(InIndex)
 {
@@ -36,6 +40,7 @@ FPQVectorIndex::~FPQVectorIndex()
 void FPQVectorIndex::Train(TConstArrayView<float> Vectors, int64 NumVectors)
 {
 	check(Vectors.Num() == NumVectors * Dimension);
+	FWriteScopeLock WriteLock(IndexLock);
 	IndexWithIDs->train(NumVectors, Vectors.GetData());
 }
 
@@ -44,31 +49,50 @@ bool FPQVectorIndex::IsTrained() const
 	return IndexWithIDs->is_trained;
 }
 
-void FPQVectorIndex::Add(int64 ID, TConstArrayView<float> Vector)
-{
-	check(IsTrained());
-	check(Vector.Num() == Dimension);
-	faiss::idx_t FaissID = static_cast<faiss::idx_t>(ID);
-	IndexWithIDs->add_with_ids(1, Vector.GetData(), &FaissID);
-}
-
-void FPQVectorIndex::AddVectors(TConstArrayView<int64> IDs, TConstArrayView<float> Vectors)
+void FPQVectorIndex::Add(TConstArrayView<int64> IDs, TConstArrayView<float> Vectors)
 {
 	check(IsTrained());
 	const int64 Count = IDs.Num();
 	check(Vectors.Num() == Count * Dimension);
 	static_assert(sizeof(faiss::idx_t) == sizeof(int64), "FAISS idx_t size mismatch");
+	FWriteScopeLock WriteLock(IndexLock);
 	IndexWithIDs->add_with_ids(
 		Count,
 		Vectors.GetData(),
 		reinterpret_cast<const faiss::idx_t*>(IDs.GetData()));
 }
 
-bool FPQVectorIndex::Remove(int64 ID)
+void FPQVectorIndex::Remove(TConstArrayView<int64> IDs)
 {
-	faiss::idx_t FaissID = static_cast<faiss::idx_t>(ID);
-	faiss::IDSelectorBatch Selector(1, &FaissID);
-	return IndexWithIDs->remove_ids(Selector) > 0;
+	if (IDs.Num() == 0)
+	{
+		return;
+	}
+	static_assert(sizeof(faiss::idx_t) == sizeof(int64), "FAISS idx_t size mismatch");
+	faiss::IDSelectorBatch Selector(IDs.Num(), reinterpret_cast<const faiss::idx_t*>(IDs.GetData()));
+	FWriteScopeLock WriteLock(IndexLock);
+	IndexWithIDs->remove_ids(Selector);
+}
+
+void FPQVectorIndex::AddQuantized(TConstArrayView<int64> IDs, TConstArrayView<uint8> Codes)
+{
+	check(IsTrained());
+	const int64 Count = IDs.Num();
+	faiss::IndexPQ* PQIndex = static_cast<faiss::IndexPQ*>(IndexWithIDs->index);
+	const int64 CodeSize = PQIndex->pq.code_size;
+	check(Codes.Num() == Count * CodeSize);
+
+	FWriteScopeLock WriteLock(IndexLock);
+
+	// Insert codes directly into the PQ index's internal storage
+	PQIndex->codes.insert(PQIndex->codes.end(), Codes.GetData(), Codes.GetData() + Codes.Num());
+	PQIndex->ntotal += Count;
+
+	// Insert IDs into the IDMap
+	static_assert(sizeof(faiss::idx_t) == sizeof(int64), "FAISS idx_t size mismatch");
+	const faiss::idx_t* IDsPtr = reinterpret_cast<const faiss::idx_t*>(IDs.GetData());
+	IndexWithIDs->id_map.insert(IndexWithIDs->id_map.end(), IDsPtr, IDsPtr + Count);
+	IndexWithIDs->ntotal += Count;
 }
 
 TArray<FSearchResult> FPQVectorIndex::Search(TConstArrayView<float> QueryVector, int32 K, TConstArrayView<int64> IDFilter) const
@@ -80,6 +104,8 @@ TArray<FSearchResult> FPQVectorIndex::Search(TConstArrayView<float> QueryVector,
 	Distances.SetNumUninitialized(K);
 	Labels.SetNumUninitialized(K);
 
+	FReadScopeLock ReadLock(IndexLock);
+
 	if (IDFilter.Num() > 0)
 	{
 		static_assert(sizeof(faiss::idx_t) == sizeof(int64), "FAISS idx_t size mismatch");
@@ -108,25 +134,12 @@ TArray<FSearchResult> FPQVectorIndex::Search(TConstArrayView<float> QueryVector,
 	return Results;
 }
 
-TArray<uint8> FPQVectorIndex::Quantize(TConstArrayView<float> Vector) const
-{
-	check(IsTrained());
-	check(Vector.Num() == Dimension);
-
-	faiss::IndexPQ* PQIndex = static_cast<faiss::IndexPQ*>(IndexWithIDs->index);
-	const int64 CodeSize = PQIndex->pq.code_size;
-
-	TArray<uint8> Codes;
-	Codes.SetNumUninitialized(CodeSize);
-	PQIndex->pq.compute_codes(Vector.GetData(), Codes.GetData(), 1);
-	return Codes;
-}
-
-TArray<uint8> FPQVectorIndex::QuantizeBatch(TConstArrayView<float> Vectors, int64 NumVectors) const
+TArray<uint8> FPQVectorIndex::Quantize(TConstArrayView<float> Vectors, int64 NumVectors) const
 {
 	check(IsTrained());
 	check(Vectors.Num() == NumVectors * Dimension);
 
+	FReadScopeLock ReadLock(IndexLock);
 	faiss::IndexPQ* PQIndex = static_cast<faiss::IndexPQ*>(IndexWithIDs->index);
 	const int64 CodeSize = PQIndex->pq.code_size;
 
@@ -138,11 +151,13 @@ TArray<uint8> FPQVectorIndex::QuantizeBatch(TConstArrayView<float> Vectors, int6
 
 int64 FPQVectorIndex::GetCount() const
 {
+	FReadScopeLock ReadLock(IndexLock);
 	return IndexWithIDs->ntotal;
 }
 
 TArray<uint8> FPQVectorIndex::Serialize() const
 {
+	FReadScopeLock ReadLock(IndexLock);
 	faiss::VectorIOWriter Writer;
 	faiss::write_index(IndexWithIDs, &Writer);
 
@@ -156,6 +171,7 @@ TArray<uint8> FPQVectorIndex::Serialize() const
 TArray<uint8> FPQVectorIndex::SerializeCodebook() const
 {
 	check(IsTrained());
+	FReadScopeLock ReadLock(IndexLock);
 	faiss::IndexPQ* PQIndex = static_cast<faiss::IndexPQ*>(IndexWithIDs->index);
 
 	faiss::VectorIOWriter Writer;
@@ -182,6 +198,7 @@ TUniquePtr<FPQVectorIndex> FPQVectorIndex::DeserializeCodebook(TConstArrayView<u
 
 	const int32 LoadedM = static_cast<int32>(PQ->M);
 	const int32 LoadedNBits = static_cast<int32>(PQ->nbits);
+	const int32 LoadedSubvectorSize = ExpectedDimension / LoadedM;
 
 	faiss::IndexPQ* PQIndex = new faiss::IndexPQ(ExpectedDimension, LoadedM, LoadedNBits);
 	PQIndex->pq = std::move(*PQ);
@@ -189,7 +206,7 @@ TUniquePtr<FPQVectorIndex> FPQVectorIndex::DeserializeCodebook(TConstArrayView<u
 	delete PQ;
 
 	faiss::IndexIDMap* IDMapIndex = new faiss::IndexIDMap(PQIndex);
-	return TUniquePtr<FPQVectorIndex>(new FPQVectorIndex(ExpectedDimension, LoadedM, LoadedNBits, IDMapIndex));
+	return TUniquePtr<FPQVectorIndex>(new FPQVectorIndex(ExpectedDimension, LoadedSubvectorSize, LoadedNBits, IDMapIndex));
 }
 
 TUniquePtr<FPQVectorIndex> FPQVectorIndex::Deserialize(TConstArrayView<uint8> Data, int32 ExpectedDimension)
@@ -208,8 +225,9 @@ TUniquePtr<FPQVectorIndex> FPQVectorIndex::Deserialize(TConstArrayView<uint8> Da
 	faiss::IndexPQ* PQIndex = static_cast<faiss::IndexPQ*>(IDMapIndex->index);
 	const int32 LoadedM = static_cast<int32>(PQIndex->pq.M);
 	const int32 LoadedNBits = static_cast<int32>(PQIndex->pq.nbits);
+	const int32 LoadedSubvectorSize = ExpectedDimension / LoadedM;
 
-	return TUniquePtr<FPQVectorIndex>(new FPQVectorIndex(ExpectedDimension, LoadedM, LoadedNBits, IDMapIndex));
+	return TUniquePtr<FPQVectorIndex>(new FPQVectorIndex(ExpectedDimension, LoadedSubvectorSize, LoadedNBits, IDMapIndex));
 }
 
 } // namespace UE::SemanticSearch
diff --git a/Engine/Plugins/Experimental/SemanticSearch/Source/SemanticSearch/Private/VectorIndex/PQVectorIndex.h b/Engine/Plugins/Experimental/SemanticSearch/Source/SemanticSearch/Private/VectorIndex/PQVectorIndex.h
--- a/Engine/Plugins/Experimental/SemanticSearch/Source/SemanticSearch/Private/VectorIndex/PQVectorIndex.h
+++ b/Engine/Plugins/Experimental/SemanticSearch/Source/SemanticSearch/Private/VectorIndex/PQVectorIndex.h
@@ -3,6 +3,7 @@
 #pragma once
 
 #include "Interfaces/IVectorIndex.h"
+#include "HAL/CriticalSection.h"
 #include "Templates/UniquePtr.h"
 
 THIRD_PARTY_INCLUDES_START
@@ -16,37 +17,44 @@ namespace UE::SemanticSearch
 /**
  * Product Quantization vector index.
  *
- * Compresses each vector into M bytes (with nbits=8), giving ~32x compression
- * for 768-dim vectors (3072 bytes -> 96 bytes). Must be trained before use.
+ * Compresses each vector into M subvectors of SubvectorSize dimensions each,
+ * quantized to nbits per subvector. With nbits=8 and SubvectorSize=8, a 768-dim
+ * vector compresses from 3072 bytes to 96 bytes (~32x).
+ *
+ * Must be trained before use.
+ *
+ * @param InSubvectorSize  Number of dimensions per subquantizer. Dimension must be divisible by this.
+ * @param InNBits          Bits per subquantizer code (8 = 256 centroids per subvector).
  */
 class FPQVectorIndex : public IVectorIndex
 {
 public:
-	FPQVectorIndex(int32 InDimension, int32 InM = 96, int32 InNBits = 8);
+	FPQVectorIndex(int32 InDimension, int32 InSubvectorSize = 8, int32 InNBits = 8);
 	virtual ~FPQVectorIndex() override;
 
 	FPQVectorIndex(const FPQVectorIndex&) = delete;
 	FPQVectorIndex& operator=(const FPQVectorIndex&) = delete;
 
-	// IVectorIndex
+	// IVectorIndex batch overrides
 	virtual void Train(TConstArrayView<float> Vectors, int64 NumVectors) override;
 	virtual bool IsTrained() const override;
-	virtual void Add(int64 ID, TConstArrayView<float> Vector) override;
-	virtual bool Remove(int64 ID) override;
+	virtual void Add(TConstArrayView<int64> IDs, TConstArrayView<float> Vectors) override;
+	virtual void Remove(TConstArrayView<int64> IDs) override;
 	virtual TArray<FSearchResult> Search(TConstArrayView<float> QueryVector, int32 K, TConstArrayView<int64> IDFilter = {}) const override;
 	virtual int64 GetCount() const override;
 	virtual int32 GetDimension() const override { return Dimension; }
 	virtual TArray<uint8> Serialize() const override;
-	virtual TArray<uint8> Quantize(TConstArrayView<float> Vector) const override;
+	virtual TArray<uint8> Quantize(TConstArrayView<float> Vectors, int64 NumVectors) const override;
 	virtual bool SupportsQuantization() const override { return true; }
+	virtual void AddQuantized(TConstArrayView<int64> IDs, TConstArrayView<uint8> Codes) override;
 
-	/** Batch add multiple vectors with corresponding IDs. */
-	void AddVectors(TConstArrayView<int64> IDs, TConstArrayView<float> Vectors);
+	/** Number of subquantizers (Dimension / SubvectorSize). */
+	int32 GetNumSubquantizers() const { return NumSubquantizers; }
 
-	/** Batch quantize vectors without adding them. Returns NumVectors * M bytes. */
-	TArray<uint8> QuantizeBatch(TConstArrayView<float> Vectors, int64 NumVectors) const;
+	/** Dimensions per subquantizer. */
+	int32 GetSubvectorSize() const { return SubvectorSize; }
 
-	int32 GetM() const { return M; }
+	/** Bits per subquantizer code. */
 	int32 GetNBits() const { return NBits; }
 
 	/** Serialize just the PQ codebook (centroids, M, nbits). */
@@ -59,12 +67,14 @@ public:
 	static TUniquePtr<FPQVectorIndex> Deserialize(TConstArrayView<uint8> Data, int32 ExpectedDimension);
 
 private:
-	FPQVectorIndex(int32 InDimension, int32 InM, int32 InNBits, faiss::IndexIDMap* InIndex);
+	FPQVectorIndex(int32 InDimension, int32 InSubvectorSize, int32 InNBits, faiss::IndexIDMap* InIndex);
 
 	int32 Dimension;
-	int32 M;
+	int32 SubvectorSize;
+	int32 NumSubquantizers;  // = Dimension / SubvectorSize
 	int32 NBits;
 	faiss::IndexIDMap* IndexWithIDs = nullptr;
+	mutable FRWLock IndexLock;
 };
 
 } // namespace UE::SemanticSearch
diff --git a/Engine/Plugins/Experimental/SemanticSearch/Source/SemanticSearch/Public/Interfaces/IVectorIndex.h b/Engine/Plugins/Experimental/SemanticSearch/Source/SemanticSearch/Public/Interfaces/IVectorIndex.h
--- a/Engine/Plugins/Experimental/SemanticSearch/Source/SemanticSearch/Public/Interfaces/IVectorIndex.h
+++ b/Engine/Plugins/Experimental/SemanticSearch/Source/SemanticSearch/Public/Interfaces/IVectorIndex.h
@@ -4,6 +4,7 @@
 
 #include "Containers/Array.h"
 #include "Containers/ArrayView.h"
+#include "Misc/AssertionMacros.h"
 
 namespace UE::SemanticSearch
 {
@@ -17,7 +18,14 @@ struct FSearchResult
 /**
  * Abstract interface for vector similarity search indices.
  *
- * Concrete implementations include:
+ * All mutation methods are batch-oriented (take arrays). Non-virtual inline
+ * helpers provide single-item convenience wrappers.
+ *
+ * Quantization methods (Quantize, AddQuantized) check(SupportsQuantization())
+ * and will crash if called on an index that doesn't support it. Callers must
+ * check SupportsQuantization() before calling.
+ *
+ * Concrete implementations:
  *   - FFlatVectorIndex: Exact brute-force search (IndexFlatL2). No training needed.
  *   - FPQVectorIndex: Product Quantization compressed search. Requires training.
  */
@@ -26,17 +34,19 @@ class IVectorIndex
 public:
 	virtual ~IVectorIndex() = default;
 
+	// --- Virtual batch methods (implementations override these) ---
+
 	/** Train the index on representative vectors. No-op for flat indices. */
 	virtual void Train(TConstArrayView<float> Vectors, int64 NumVectors) = 0;
 
 	/** Whether the index has been trained and is ready for Add/Search. */
 	virtual bool IsTrained() const = 0;
 
-	/** Add a single vector with an external ID. Index must be trained first. */
-	virtual void Add(int64 ID, TConstArrayView<float> Vector) = 0;
+	/** Add multiple vectors with corresponding IDs. Index must be trained first. */
+	virtual void Add(TConstArrayView<int64> IDs, TConstArrayView<float> Vectors) = 0;
 
-	/** Remove a vector by ID. Returns true if found and removed. */
-	virtual bool Remove(int64 ID) = 0;
+	/** Remove vectors by IDs. */
+	virtual void Remove(TConstArrayView<int64> IDs) = 0;
 
 	/**
 	 * Search for K nearest neighbors. Returns results sorted by ascending distance.
@@ -53,11 +63,51 @@ public:
 	/** Serialize the full index to a byte array. */
 	virtual TArray<uint8> Serialize() const = 0;
 
-	/** Quantize a vector into compressed form (e.g. PQ codes). Returns empty if not supported. */
-	virtual TArray<uint8> Quantize(TConstArrayView<float> Vector) const { return {}; }
-
-	/** Whether this index supports quantizing vectors into a compressed representation. */
+	/** Whether this index supports quantization (Quantize/AddQuantized). */
 	virtual bool SupportsQuantization() const { return false; }
+
+	/**
+	 * Quantize vectors into compressed form (e.g. PQ codes).
+	 * check(SupportsQuantization()) — crashes if not supported.
+	 */
+	virtual TArray<uint8> Quantize(TConstArrayView<float> Vectors, int64 NumVectors) const
+	{
+		checkf(false, TEXT("Quantize called on index that does not support quantization"));
+		return {};
+	}
+
+	/**
+	 * Add pre-quantized vectors directly, skipping encoding.
+	 * check(SupportsQuantization()) — crashes if not supported.
+	 */
+	virtual void AddQuantized(TConstArrayView<int64> IDs, TConstArrayView<uint8> Codes)
+	{
+		checkf(false, TEXT("AddQuantized called on index that does not support quantization"));
+	}
+
+	// --- Inline single-item helpers ---
+
+	void Add(int64 ID, TConstArrayView<float> Vector)
+	{
+		Add(TConstArrayView<int64>(&ID, 1), Vector);
+	}
+
+	void AddQuantized(int64 ID, TConstArrayView<uint8> Codes)
+	{
+		AddQuantized(TConstArrayView<int64>(&ID, 1), Codes);
+	}
+
+	bool Remove(int64 ID)
+	{
+		const int64 CountBefore = GetCount();
+		Remove(TConstArrayView<int64>(&ID, 1));
+		return GetCount() < CountBefore;
+	}
+
+	TArray<uint8> Quantize(TConstArrayView<float> Vector) const
+	{
+		return Quantize(Vector, 1);
+	}
 };
 
 } // namespace UE::SemanticSearch
