Change RemoteEmbeddingProvider to use new endpoints

#Synced-CL 52250376

diff --git a/Engine/Plugins/Experimental/SemanticSearch/Source/SemanticSearch/Private/RemoteEmbeddingProvider.cpp b/Engine/Plugins/Experimental/SemanticSearch/Source/SemanticSearch/Private/RemoteEmbeddingProvider.cpp
--- a/Engine/Plugins/Experimental/SemanticSearch/Source/SemanticSearch/Private/RemoteEmbeddingProvider.cpp
+++ b/Engine/Plugins/Experimental/SemanticSearch/Source/SemanticSearch/Private/RemoteEmbeddingProvider.cpp
@@ -98,14 +98,20 @@ private:
 };
 
 
-static FString ReadBaseURLFromConfig()
+struct FServiceURLs
+{
+	FString EmbedURL;
+	FString CaptionURL;
+};
+
+static FServiceURLs ReadURLsFromConfig()
 {
 	const FString ConfigPath = FPaths::EngineDir() / TEXT("Restricted/NotForLicensees/Config/SemanticSearch.json");
 	FString FileContents;
 	if (!FFileHelper::LoadFileToString(FileContents, *ConfigPath))
 	{
 		UE_LOG(LogSemanticSearch, Warning, TEXT("Failed to load config: %s"), *ConfigPath);
-		return FString();
+		return {};
 	}
 
 	TSharedPtr<FJsonObject> Json;
@@ -113,22 +119,28 @@ static FString ReadBaseURLFromConfig()
 	if (!FJsonSerializer::Deserialize(Reader, Json) || !Json.IsValid())
 	{
 		UE_LOG(LogSemanticSearch, Warning, TEXT("Failed to parse config JSON: %s"), *ConfigPath);
-		return FString();
+		return {};
 	}
 
-	FString URL;
-	if (!Json->TryGetStringField(TEXT("base_url"), URL))
+	FServiceURLs URLs;
+	if (!Json->TryGetStringField(TEXT("embed_url"), URLs.EmbedURL))
 	{
-		UE_LOG(LogSemanticSearch, Warning, TEXT("Missing 'base_url' in config: %s"), *ConfigPath);
+		UE_LOG(LogSemanticSearch, Warning, TEXT("Missing 'embed_url' in config: %s"), *ConfigPath);
 	}
-	return URL;
+	if (!Json->TryGetStringField(TEXT("caption_url"), URLs.CaptionURL))
+	{
+		UE_LOG(LogSemanticSearch, Warning, TEXT("Missing 'caption_url' in config: %s"), *ConfigPath);
+	}
+	return URLs;
 }
 
-FRemoteEmbeddingProvider::FRemoteEmbeddingProvider(const FString& InBaseURL)
-	: BaseURL(InBaseURL.IsEmpty() ? ReadBaseURLFromConfig() : InBaseURL)
-	, EmbeddingDimension(0)
+FRemoteEmbeddingProvider::FRemoteEmbeddingProvider()
+	: EmbeddingDimension(0)
 {
-	UE_LOG(LogSemanticSearch, Log, TEXT("RemoteEmbeddingProvider created (URL: %s)"), *BaseURL);
+	FServiceURLs URLs = ReadURLsFromConfig();
+	EmbedURL = MoveTemp(URLs.EmbedURL);
+	CaptionURL = MoveTemp(URLs.CaptionURL);
+	UE_LOG(LogSemanticSearch, Log, TEXT("RemoteEmbeddingProvider created (Embed: %s, Caption: %s)"), *EmbedURL, *CaptionURL);
 }
 
 int32 FRemoteEmbeddingProvider::GetEmbeddingDimension() const
@@ -145,7 +157,7 @@ void FRemoteEmbeddingProvider::GenerateCaptionAsync(
 	TArray<uint8> Body = BuildMultipartBody(Request, Boundary);
 
 	TSharedRef<IHttpRequest, ESPMode::ThreadSafe> Http = FHttpModule::Get().CreateRequest();
-	Http->SetURL(BaseURL / TEXT("v1") / TEXT("caption"));
+	Http->SetURL(CaptionURL);
 	Http->SetVerb(TEXT("POST"));
 	Http->SetHeader(TEXT("Content-Type"),
 		FString::Printf(TEXT("multipart/form-data; boundary=%s"), *Boundary));
@@ -197,15 +209,15 @@ void FRemoteEmbeddingProvider::GenerateCaptionAsync(
 				return;
 			}
 
-			const TArray<TSharedPtr<FJsonValue>>* KeywordsArray;
-			if (Json->TryGetArrayField(TEXT("keywords"), KeywordsArray))
+			const TArray<TSharedPtr<FJsonValue>>* TagsArray;
+			if (Json->TryGetArrayField(TEXT("tags"), TagsArray))
 			{
-				for (const TSharedPtr<FJsonValue>& Val : *KeywordsArray)
+				for (const TSharedPtr<FJsonValue>& Val : *TagsArray)
 				{
-					FString Kw;
-					if (Val->TryGetString(Kw))
+					FString Tag;
+					if (Val->TryGetString(Tag))
 					{
-						Response.Keywords.Add(Kw);
+						Response.Keywords.Add(Tag);
 					}
 				}
 			}
@@ -238,7 +250,7 @@ void FRemoteEmbeddingProvider::GenerateEmbeddingAsync(
 	FJsonSerializer::Serialize(Payload, Writer);
 
 	TSharedRef<IHttpRequest, ESPMode::ThreadSafe> Http = FHttpModule::Get().CreateRequest();
-	Http->SetURL(BaseURL / TEXT("v1") / TEXT("embed"));
+	Http->SetURL(EmbedURL);
 	Http->SetVerb(TEXT("POST"));
 	Http->SetHeader(TEXT("Content-Type"), TEXT("application/json"));
 	Http->SetContentAsString(PayloadStr);
@@ -314,7 +326,11 @@ void FRemoteEmbeddingProvider::GenerateEmbeddingAsync(
 					Response.Embedding.Num());
 			}
 
-			Json->TryGetStringField(TEXT("model_version"), Response.ModelVersion);
+			const TSharedPtr<FJsonObject>* MetadataObj;
+			if (Json->TryGetObjectField(TEXT("metadata"), MetadataObj))
+			{
+				(*MetadataObj)->TryGetStringField(TEXT("model_name"), Response.ModelVersion);
+			}
 
 			OnComplete->operator()(MoveTemp(Response));
 		});
@@ -345,7 +361,7 @@ TArray<uint8> FRemoteEmbeddingProvider::BuildMultipartBody(
 		Body.Append(reinterpret_cast<const uint8*>(Utf8.Get()), Utf8.Length());
 	};
 
-	// media[] file parts + media_types[] form fields
+	// media file parts
 	for (int32 i = 0; i < Request.AssetMedia.Num(); ++i)
 	{
 		const FAssetMedia& M = Request.AssetMedia[i];
@@ -358,40 +374,43 @@ TArray<uint8> FRemoteEmbeddingProvider::BuildMultipartBody(
 		else if (M.MimeType.StartsWith(TEXT("text/")))     Ext = TEXT("txt");
 
 		Append(Delim);
-		Append(FString::Printf(TEXT("Content-Disposition: form-data; name=\"media[]\"; filename=\"media_%d.%s\"\r\n"), i, *Ext));
+		Append(FString::Printf(TEXT("Content-Disposition: form-data; name=\"media\"; filename=\"media_%d.%s\"\r\n"), i, *Ext));
 		Append(FString::Printf(TEXT("Content-Type: %s\r\n\r\n"), *M.MimeType));
 		Body.Append(M.Data.GetData(), M.Data.Num());
 		Append(TEXT("\r\n"));
-
-		Append(Delim);
-		Append(TEXT("Content-Disposition: form-data; name=\"media_types[]\"\r\n\r\n"));
-		Append(FString::Printf(TEXT("%s\r\n"), *M.MimeType));
 	}
 
-	if (!Request.AssetPath.IsEmpty())
-	{
-		Append(Delim);
-		Append(TEXT("Content-Disposition: form-data; name=\"asset_path\"\r\n\r\n"));
-		Append(FString::Printf(TEXT("%s\r\n"), *Request.AssetPath));
-	}
+	// task_preset
+	Append(Delim);
+	Append(TEXT("Content-Disposition: form-data; name=\"task_preset\"\r\n\r\n"));
+	Append(FString::Printf(TEXT("%s\r\n"), TaskPreset));
 
-	if (!Request.AssetType.IsEmpty())
+	// metadata — merge asset_path and asset_class into the metadata JSON
 	{
-		Append(Delim);
-		Append(TEXT("Content-Disposition: form-data; name=\"asset_class\"\r\n\r\n"));
-		Append(FString::Printf(TEXT("%s\r\n"), *Request.AssetType));
-	}
+		TSharedPtr<FJsonObject> MetadataJson = Request.Metadata.IsValid()
+			? MakeShared<FJsonObject>(*Request.Metadata)
+			: MakeShared<FJsonObject>();
 
-	if (Request.Metadata.IsValid() && Request.Metadata->Values.Num() > 0)
-	{
-		FString MetaStr;
-		TSharedRef<TJsonWriter<TCHAR, TCondensedJsonPrintPolicy<TCHAR>>> Wr =
-			TJsonWriterFactory<TCHAR, TCondensedJsonPrintPolicy<TCHAR>>::Create(&MetaStr);
-		FJsonSerializer::Serialize(Request.Metadata.ToSharedRef(), Wr);
+		if (!Request.AssetPath.IsEmpty())
+		{
+			MetadataJson->SetStringField(TEXT("asset_path"), Request.AssetPath);
+		}
+		if (!Request.AssetType.IsEmpty())
+		{
+			MetadataJson->SetStringField(TEXT("asset_class"), Request.AssetType);
+		}
 
-		Append(Delim);
-		Append(TEXT("Content-Disposition: form-data; name=\"metadata\"\r\n\r\n"));
-		Append(FString::Printf(TEXT("%s\r\n"), *MetaStr));
+		if (MetadataJson->Values.Num() > 0)
+		{
+			FString MetaStr;
+			TSharedRef<TJsonWriter<TCHAR, TCondensedJsonPrintPolicy<TCHAR>>> Wr =
+				TJsonWriterFactory<TCHAR, TCondensedJsonPrintPolicy<TCHAR>>::Create(&MetaStr);
+			FJsonSerializer::Serialize(MetadataJson.ToSharedRef(), Wr);
+
+			Append(Delim);
+			Append(TEXT("Content-Disposition: form-data; name=\"metadata\"\r\n\r\n"));
+			Append(FString::Printf(TEXT("%s\r\n"), *MetaStr));
+		}
 	}
 
 	Append(End);
diff --git a/Engine/Plugins/Experimental/SemanticSearch/Source/SemanticSearch/Private/RemoteEmbeddingProvider.h b/Engine/Plugins/Experimental/SemanticSearch/Source/SemanticSearch/Private/RemoteEmbeddingProvider.h
--- a/Engine/Plugins/Experimental/SemanticSearch/Source/SemanticSearch/Private/RemoteEmbeddingProvider.h
+++ b/Engine/Plugins/Experimental/SemanticSearch/Source/SemanticSearch/Private/RemoteEmbeddingProvider.h
@@ -10,27 +10,28 @@
 namespace UE::SemanticSearch::Private
 {
 /**
- *  Remote embedding provider that communicates with Python backend.
+ *  Remote embedding provider that communicates with MLS backend services.
  *
- *  GenerateCaptionAsync   -> POST /v1/caption  (multipart/form-data)
- *  GenerateEmbeddingAsync -> POST /v1/embed    (application/json)
+ *  GenerateCaptionAsync   -> POST caption_url  (multipart/form-data)
+ *  GenerateEmbeddingAsync -> POST embed_url    (application/json)
  */
 class FRemoteEmbeddingProvider : public IEmbeddingProvider
 {
 public:
-	explicit FRemoteEmbeddingProvider(const FString& InBaseURL = FString());
+	FRemoteEmbeddingProvider();
 
 	virtual int32 GetEmbeddingDimension() const override;
 	virtual void GenerateCaptionAsync(FCaptionRequest&& Request, FOnCaptionComplete&& OnComplete, DerivedData::FRequestOwner* RequestOwner) override;
 	virtual void GenerateEmbeddingAsync(const FStringView Text, FOnEmbeddingComplete&& OnComplete, DerivedData::FRequestOwner* RequestOwner) override;
 
-	FString GetBaseURL() const { return BaseURL; }
-
 private:
 	/** Build multipart/form-data body.*/
 	TArray<uint8> BuildMultipartBody(const FCaptionRequest& Request, FString& OutBoundary) const;
 
-	FString BaseURL;
+	static constexpr const TCHAR* TaskPreset = TEXT("UE_ASSET_TAG_AND_CAPTION_V1");
+
+	FString EmbedURL;
+	FString CaptionURL;
 	TAtomic<int32> EmbeddingDimension;
 };
 }
