Moving GetProcessSubBucketName implementations to headers since constexpr only works if callers can see the value.
This was causing compile errors in certain build modes.

#Synced-CL 52214067

diff --git a/Engine/Plugins/Experimental/SemanticSearch/Source/SemanticSearch/Private/AssetProcessors/MaterialAssetProcessor.cpp b/Engine/Plugins/Experimental/SemanticSearch/Source/SemanticSearch/Private/AssetProcessors/MaterialAssetProcessor.cpp
--- a/Engine/Plugins/Experimental/SemanticSearch/Source/SemanticSearch/Private/AssetProcessors/MaterialAssetProcessor.cpp
+++ b/Engine/Plugins/Experimental/SemanticSearch/Source/SemanticSearch/Private/AssetProcessors/MaterialAssetProcessor.cpp
@@ -12,11 +12,6 @@
 namespace UE::SemanticSearch::Private
 {
 
-constexpr FStringView FMaterialProcessor::GetProcessSubBucketName() const
-{
-	return TEXTVIEW("FMaterialProcessor");
-}
-
 UClass& FMaterialProcessor::GetSupportedClass() const
 {
 	return *UMaterial::StaticClass();
diff --git a/Engine/Plugins/Experimental/SemanticSearch/Source/SemanticSearch/Private/AssetProcessors/MaterialAssetProcessor.h b/Engine/Plugins/Experimental/SemanticSearch/Source/SemanticSearch/Private/AssetProcessors/MaterialAssetProcessor.h
--- a/Engine/Plugins/Experimental/SemanticSearch/Source/SemanticSearch/Private/AssetProcessors/MaterialAssetProcessor.h
+++ b/Engine/Plugins/Experimental/SemanticSearch/Source/SemanticSearch/Private/AssetProcessors/MaterialAssetProcessor.h
@@ -13,7 +13,10 @@ public:
 	// Must be created on the game thread
 	FMaterialProcessor();
 
-	constexpr virtual FStringView GetProcessSubBucketName() const override;
+	constexpr virtual FStringView GetProcessSubBucketName() const override
+	{
+		return TEXTVIEW("FMaterialProcessor");
+	}
 
 	// Thread safe
 	virtual UClass& GetSupportedClass() const override;
diff --git a/Engine/Plugins/Experimental/SemanticSearch/Source/SemanticSearch/Private/AssetProcessors/MaterialInstanceAssetProcessor.cpp b/Engine/Plugins/Experimental/SemanticSearch/Source/SemanticSearch/Private/AssetProcessors/MaterialInstanceAssetProcessor.cpp
--- a/Engine/Plugins/Experimental/SemanticSearch/Source/SemanticSearch/Private/AssetProcessors/MaterialInstanceAssetProcessor.cpp
+++ b/Engine/Plugins/Experimental/SemanticSearch/Source/SemanticSearch/Private/AssetProcessors/MaterialInstanceAssetProcessor.cpp
@@ -12,11 +12,6 @@ namespace UE::SemanticSearch::Private
 
 FMaterialInstanceProcessor::FMaterialInstanceProcessor() = default;
 
-constexpr FStringView FMaterialInstanceProcessor::GetProcessSubBucketName() const
-{
-	return TEXTVIEW("FMaterialInstanceProcessor");
-}
-
 UClass& FMaterialInstanceProcessor::GetSupportedClass() const
 {
 	return *UMaterialInstance::StaticClass();
diff --git a/Engine/Plugins/Experimental/SemanticSearch/Source/SemanticSearch/Private/AssetProcessors/MaterialInstanceAssetProcessor.h b/Engine/Plugins/Experimental/SemanticSearch/Source/SemanticSearch/Private/AssetProcessors/MaterialInstanceAssetProcessor.h
--- a/Engine/Plugins/Experimental/SemanticSearch/Source/SemanticSearch/Private/AssetProcessors/MaterialInstanceAssetProcessor.h
+++ b/Engine/Plugins/Experimental/SemanticSearch/Source/SemanticSearch/Private/AssetProcessors/MaterialInstanceAssetProcessor.h
@@ -12,7 +12,10 @@ class FMaterialInstanceProcessor : public FThumbnailBaseAssetProcessor
 public:
 	FMaterialInstanceProcessor();
 
-	constexpr virtual FStringView GetProcessSubBucketName() const override;
+	constexpr virtual FStringView GetProcessSubBucketName() const override
+	{
+		return TEXTVIEW("FMaterialInstanceProcessor");
+	}
 
 	// Thread safe
 	virtual UClass& GetSupportedClass() const override;
diff --git a/Engine/Plugins/Experimental/SemanticSearch/Source/SemanticSearch/Private/AssetProcessors/SkeletalMeshAssetProcessor.cpp b/Engine/Plugins/Experimental/SemanticSearch/Source/SemanticSearch/Private/AssetProcessors/SkeletalMeshAssetProcessor.cpp
--- a/Engine/Plugins/Experimental/SemanticSearch/Source/SemanticSearch/Private/AssetProcessors/SkeletalMeshAssetProcessor.cpp
+++ b/Engine/Plugins/Experimental/SemanticSearch/Source/SemanticSearch/Private/AssetProcessors/SkeletalMeshAssetProcessor.cpp
@@ -12,11 +12,6 @@ namespace UE::SemanticSearch::Private
 
 FSkeletalMeshProcessor::FSkeletalMeshProcessor() = default;
 
-constexpr FStringView FSkeletalMeshProcessor::GetProcessSubBucketName() const
-{
-	return TEXTVIEW("FSkeletalMeshProcessor");
-}
-
 UClass& FSkeletalMeshProcessor::GetSupportedClass() const
 {
 	return *USkeletalMesh::StaticClass();
diff --git a/Engine/Plugins/Experimental/SemanticSearch/Source/SemanticSearch/Private/AssetProcessors/SkeletalMeshAssetProcessor.h b/Engine/Plugins/Experimental/SemanticSearch/Source/SemanticSearch/Private/AssetProcessors/SkeletalMeshAssetProcessor.h
--- a/Engine/Plugins/Experimental/SemanticSearch/Source/SemanticSearch/Private/AssetProcessors/SkeletalMeshAssetProcessor.h
+++ b/Engine/Plugins/Experimental/SemanticSearch/Source/SemanticSearch/Private/AssetProcessors/SkeletalMeshAssetProcessor.h
@@ -12,7 +12,10 @@ class FSkeletalMeshProcessor : public FThumbnailBaseAssetProcessor
 public:
 	FSkeletalMeshProcessor();
 
-	constexpr virtual FStringView GetProcessSubBucketName() const override;
+	constexpr virtual FStringView GetProcessSubBucketName() const override
+	{
+		return TEXTVIEW("FSkeletalMeshProcessor");
+	}
 
 	// Thread safe
 	virtual UClass& GetSupportedClass() const override;
diff --git a/Engine/Plugins/Experimental/SemanticSearch/Source/SemanticSearch/Private/AssetProcessors/StaticMeshAssetProcessor.cpp b/Engine/Plugins/Experimental/SemanticSearch/Source/SemanticSearch/Private/AssetProcessors/StaticMeshAssetProcessor.cpp
--- a/Engine/Plugins/Experimental/SemanticSearch/Source/SemanticSearch/Private/AssetProcessors/StaticMeshAssetProcessor.cpp
+++ b/Engine/Plugins/Experimental/SemanticSearch/Source/SemanticSearch/Private/AssetProcessors/StaticMeshAssetProcessor.cpp
@@ -12,11 +12,6 @@
 namespace UE::SemanticSearch::Private
 {
 
-constexpr FStringView FStaticMeshProcessor::GetProcessSubBucketName() const
-{
-	return TEXTVIEW("FStaticMeshProcessor");
-}
-
 UClass& FStaticMeshProcessor::GetSupportedClass() const
 {
 	return *UStaticMesh::StaticClass();
diff --git a/Engine/Plugins/Experimental/SemanticSearch/Source/SemanticSearch/Private/AssetProcessors/StaticMeshAssetProcessor.h b/Engine/Plugins/Experimental/SemanticSearch/Source/SemanticSearch/Private/AssetProcessors/StaticMeshAssetProcessor.h
--- a/Engine/Plugins/Experimental/SemanticSearch/Source/SemanticSearch/Private/AssetProcessors/StaticMeshAssetProcessor.h
+++ b/Engine/Plugins/Experimental/SemanticSearch/Source/SemanticSearch/Private/AssetProcessors/StaticMeshAssetProcessor.h
@@ -13,7 +13,10 @@ public:
 	// Must be created on the game thread
 	FStaticMeshProcessor();
 
-	constexpr virtual FStringView GetProcessSubBucketName() const override;
+	constexpr virtual FStringView GetProcessSubBucketName() const override
+	{
+		return TEXTVIEW("FStaticMeshProcessor");
+	}
 
 	// Thread safe
 	virtual UClass& GetSupportedClass() const override;
diff --git a/Engine/Plugins/Experimental/SemanticSearch/Source/SemanticSearch/Private/AssetProcessors/TextureAssetProcessor.cpp b/Engine/Plugins/Experimental/SemanticSearch/Source/SemanticSearch/Private/AssetProcessors/TextureAssetProcessor.cpp
--- a/Engine/Plugins/Experimental/SemanticSearch/Source/SemanticSearch/Private/AssetProcessors/TextureAssetProcessor.cpp
+++ b/Engine/Plugins/Experimental/SemanticSearch/Source/SemanticSearch/Private/AssetProcessors/TextureAssetProcessor.cpp
@@ -12,11 +12,6 @@
 namespace UE::SemanticSearch::Private
 {
 
-constexpr FStringView FTextureProcessor::GetProcessSubBucketName() const
-{
-	return TEXTVIEW("FTextureProcessor");
-}
-
 UClass& FTextureProcessor::GetSupportedClass() const
 {
 	return *UTexture::StaticClass();
diff --git a/Engine/Plugins/Experimental/SemanticSearch/Source/SemanticSearch/Private/AssetProcessors/TextureAssetProcessor.h b/Engine/Plugins/Experimental/SemanticSearch/Source/SemanticSearch/Private/AssetProcessors/TextureAssetProcessor.h
--- a/Engine/Plugins/Experimental/SemanticSearch/Source/SemanticSearch/Private/AssetProcessors/TextureAssetProcessor.h
+++ b/Engine/Plugins/Experimental/SemanticSearch/Source/SemanticSearch/Private/AssetProcessors/TextureAssetProcessor.h
@@ -13,7 +13,10 @@ public:
 	// Must be created on the game thread
 	FTextureProcessor();
 
-	constexpr virtual FStringView GetProcessSubBucketName() const override;
+	constexpr virtual FStringView GetProcessSubBucketName() const override
+	{
+		return TEXTVIEW("FTextureProcessor");
+	}
 
 	// Thread safe
 	virtual UClass& GetSupportedClass() const override;
