Automatic indexing callbacks on asset creation/resave/rename/delete

Also CVar settings for auto-indexing toggles, and limiting folders for indexing

Also some UE Cmd for getting Index stats and manually triggering the indexing

#rnx
#Synced-CL 52134402

diff --git a/Engine/Plugins/Experimental/SemanticSearch/Source/SemanticSearch/Private/AssetProcessorManager.cpp b/Engine/Plugins/Experimental/SemanticSearch/Source/SemanticSearch/Private/AssetProcessorManager.cpp
--- a/Engine/Plugins/Experimental/SemanticSearch/Source/SemanticSearch/Private/AssetProcessorManager.cpp
+++ b/Engine/Plugins/Experimental/SemanticSearch/Source/SemanticSearch/Private/AssetProcessorManager.cpp
@@ -87,16 +87,16 @@ namespace UE::SemanticSearch
 		});
 	}
 
-	bool FAssetProcessorManager::GetEmbeddingData(const FAssetData& InAssetData, TUniqueFunction<void(FAssetEmbeddingResult&&, FString&&)>&& InOnDataAvailable)
+	bool FAssetProcessorManager::GetEmbeddingData(const FAssetData& InAssetData, TUniqueFunction<void(FAssetEmbeddingResult&&, FString&&)>&& InOnDataAvailable, bool bBuildOnMiss)
 	{
-		return GetOrCreateTask(InAssetData, [&InOnDataAvailable](const TSharedRef<Private::FSemanticSearchBuildCacheTask>& Task)
+		return GetOrCreateTask(InAssetData, [&InOnDataAvailable, bBuildOnMiss](const TSharedRef<Private::FSemanticSearchBuildCacheTask>& Task)
 		{
 			Task->GetEmbeddingData([OnDataAvailable = MoveTemp(InOnDataAvailable)](Private::FEmbeddingDerivedData&& EmbeddingData, FString&& ErrorMessage) mutable
 			{
 				FAssetEmbeddingResult Result;
 				Result.Embedding = MoveTemp(EmbeddingData.Embedding);
 				OnDataAvailable(MoveTemp(Result), MoveTemp(ErrorMessage));
-			});
+			}, bBuildOnMiss);
 		});
 	}
 
diff --git a/Engine/Plugins/Experimental/SemanticSearch/Source/SemanticSearch/Private/AssetProcessors/AssetProcessorUtils.cpp b/Engine/Plugins/Experimental/SemanticSearch/Source/SemanticSearch/Private/AssetProcessors/AssetProcessorUtils.cpp
--- a/Engine/Plugins/Experimental/SemanticSearch/Source/SemanticSearch/Private/AssetProcessors/AssetProcessorUtils.cpp
+++ b/Engine/Plugins/Experimental/SemanticSearch/Source/SemanticSearch/Private/AssetProcessors/AssetProcessorUtils.cpp
@@ -69,11 +69,10 @@ void RegisterDefaultAssetProcessors()
 	Manager.RegisterAssetProcessor<FTextureProcessor>();
 }
 
-int64 GetAssetIndexID(const FAssetData& Asset)
+int64 GetAssetIndexID(const FString& ObjectPath)
 {
-	const FString Path = Asset.GetObjectPathString();
 	FBlake3 Hasher;
-	Hasher.Update(Path.GetCharArray().GetData(), Path.Len() * sizeof(TCHAR));
+	Hasher.Update(ObjectPath.GetCharArray().GetData(), ObjectPath.Len() * sizeof(TCHAR));
 	FBlake3Hash Hash = Hasher.Finalize();
 	int64 Result;
 	FMemory::Memcpy(&Result, Hash.GetBytes(), sizeof(int64));
@@ -81,4 +80,9 @@ int64 GetAssetIndexID(const FAssetData& Asset)
 	return Result & 0x7FFFFFFFFFFFFFFF;
 }
 
+int64 GetAssetIndexID(const FAssetData& Asset)
+{
+	return GetAssetIndexID(Asset.GetObjectPathString());
+}
+
 }
diff --git a/Engine/Plugins/Experimental/SemanticSearch/Source/SemanticSearch/Private/AssetProcessors/AssetProcessorUtils.h b/Engine/Plugins/Experimental/SemanticSearch/Source/SemanticSearch/Private/AssetProcessors/AssetProcessorUtils.h
--- a/Engine/Plugins/Experimental/SemanticSearch/Source/SemanticSearch/Private/AssetProcessors/AssetProcessorUtils.h
+++ b/Engine/Plugins/Experimental/SemanticSearch/Source/SemanticSearch/Private/AssetProcessors/AssetProcessorUtils.h
@@ -61,5 +61,6 @@ void RegisterDefaultAssetProcessors();
  * Uses Blake3. Always non-negative (FAISS reserves -1 as sentinel).
  */
 int64 GetAssetIndexID(const FAssetData& Asset);
+int64 GetAssetIndexID(const FString& ObjectPath);
 
 }
diff --git a/Engine/Plugins/Experimental/SemanticSearch/Source/SemanticSearch/Private/SemanticSearchDerivedDataBuilder.cpp b/Engine/Plugins/Experimental/SemanticSearch/Source/SemanticSearch/Private/SemanticSearchDerivedDataBuilder.cpp
--- a/Engine/Plugins/Experimental/SemanticSearch/Source/SemanticSearch/Private/SemanticSearchDerivedDataBuilder.cpp
+++ b/Engine/Plugins/Experimental/SemanticSearch/Source/SemanticSearch/Private/SemanticSearchDerivedDataBuilder.cpp
@@ -80,7 +80,7 @@ FSemanticSearchBuildCacheTask::FSemanticSearchBuildCacheTask(FAssetData&& InAsse
 {
 }
 
-void FSemanticSearchBuildCacheTask::GetEmbeddingData(TUniqueFunction<void(FEmbeddingDerivedData&& EmbeddingData, FString&& ErrorMessage)>&& InOnDataAvailable)
+void FSemanticSearchBuildCacheTask::GetEmbeddingData(TUniqueFunction<void(FEmbeddingDerivedData&& EmbeddingData, FString&& ErrorMessage)>&& InOnDataAvailable, bool bBuildOnMiss)
 {
 	TRACE_CPUPROFILER_EVENT_SCOPE_STR("SemanticSearch::GetEmbeddingData");
 	{
@@ -103,7 +103,7 @@ void FSemanticSearchBuildCacheTask::GetEmbeddingData(TUniqueFunction<void(FEmbed
 	Request.Policy = GetCachePolicy();
 
 	FRequestBarrier Barrier(*this);
-	GetCache().GetChunks({ Request }, *this, [this](FCacheGetChunkResponse&& Response) mutable
+	GetCache().GetChunks({ Request }, *this, [this, bBuildOnMiss](FCacheGetChunkResponse&& Response) mutable
 		{
 			if (Response.Status == EStatus::Ok && Response.RawData.GetSize() > 0)
 			{
@@ -119,7 +119,14 @@ void FSemanticSearchBuildCacheTask::GetEmbeddingData(TUniqueFunction<void(FEmbed
 			else if (Response.Status != EStatus::Canceled)
 			{
 				TRACE_CPUPROFILER_EVENT_SCOPE_STR("SemanticSearch::GetEmbeddingData::DDCMiss");
-				this->BuildDerivedData();
+				if (bBuildOnMiss)
+				{
+					this->BuildDerivedData();
+				}
+				else
+				{
+					BroadcastQueuedEmbeddingRequest(FEmbeddingDerivedData(), FString(TEXT("Not cached")));
+				}
 			}
 			else
 			{
@@ -207,7 +214,7 @@ void FSemanticSearchBuildCacheTask::CacheData()
 
 DerivedData::FCacheKey FSemanticSearchBuildCacheTask::GenerateCacheKey(const FAssetData& InAssetData, const TSharedRef<IAssetProcessor>& InAssetProcessor)
 {
-	constexpr FGuid SemanticSearchDerivedDataGuid(0xA3D4A6CE, 0xB6C04E41, 0xABC5B4C2, 0x999D0095);
+	constexpr FGuid SemanticSearchDerivedDataGuid(0x538A7DB0, 0x80304195, 0x84016942, 0xD8879E0A);
 	FMemoryHasherBlake3 HashWriter;
 
 	// The Hasher won't modify the value
diff --git a/Engine/Plugins/Experimental/SemanticSearch/Source/SemanticSearch/Private/SemanticSearchDerivedDataBuilder.h b/Engine/Plugins/Experimental/SemanticSearch/Source/SemanticSearch/Private/SemanticSearchDerivedDataBuilder.h
--- a/Engine/Plugins/Experimental/SemanticSearch/Source/SemanticSearch/Private/SemanticSearchDerivedDataBuilder.h
+++ b/Engine/Plugins/Experimental/SemanticSearch/Source/SemanticSearch/Private/SemanticSearchDerivedDataBuilder.h
@@ -49,7 +49,7 @@ public:
 
 	FSemanticSearchBuildCacheTask(FAssetData&& InAssetData, const TSharedRef<IAssetProcessor>& InAssetProcessor, DerivedData::FCacheKey&& InAssetKey);
 
-	void GetEmbeddingData(TUniqueFunction<void (FEmbeddingDerivedData&& EmbeddingData, FString&& ErrorMessage)>&& InOnDataAvailable);
+	void GetEmbeddingData(TUniqueFunction<void (FEmbeddingDerivedData&& EmbeddingData, FString&& ErrorMessage)>&& InOnDataAvailable, bool bBuildOnMiss = true);
 	void GetCaptionData(TUniqueFunction<void (FCaptionDerivedData&& CaptionData, FString&& ErrorMessage)>&& InOnDataAvailable);
 
 	void CacheData();
diff --git a/Engine/Plugins/Experimental/SemanticSearch/Source/SemanticSearch/Private/SemanticSearchModule.cpp b/Engine/Plugins/Experimental/SemanticSearch/Source/SemanticSearch/Private/SemanticSearchModule.cpp
--- a/Engine/Plugins/Experimental/SemanticSearch/Source/SemanticSearch/Private/SemanticSearchModule.cpp
+++ b/Engine/Plugins/Experimental/SemanticSearch/Source/SemanticSearch/Private/SemanticSearchModule.cpp
@@ -8,6 +8,8 @@
 
 #include "Modules/ModuleManager.h"
 #include "RemoteEmbeddingProvider.h"
+#include "AssetRegistry/AssetRegistryModule.h"
+#include "HAL/ConsoleManager.h"
 
 DEFINE_LOG_CATEGORY(LogSemanticSearch);
 
@@ -15,12 +17,55 @@ namespace UE::SemanticSearch::Private
 {
 IMPLEMENT_MODULE(FSemanticSearchModule, SemanticSearch)
 
+static TAutoConsoleVariable<bool> CVarAutoIndexOnStartup(
+	TEXT("SemanticSearch.AutoIndexOnStartup"),
+	false,
+	TEXT("When true, assets missing from the DDC cache are sent to the backend for indexing on editor startup. DDC hits always populate the FAISS index regardless."));
+
+static TAutoConsoleVariable<bool> CVarAutoIndexOnSave(
+	TEXT("SemanticSearch.AutoIndexOnSave"),
+	true,
+	TEXT("When true, assets are automatically re-indexed when saved or updated on disk."));
+
+static TAutoConsoleVariable<FString> CVarIndexedPaths(
+	TEXT("SemanticSearch.IndexedPaths"),
+	TEXT("/Game/SemanticSearchTest/"),
+	TEXT("Comma-separated list of content browser paths to include in the semantic search index. E.g. /Game/Folder1/,/Game/Folder2/"));
+
+static TArray<FString> GetIndexedFolderPaths()
+{
+	TArray<FString> Paths;
+	CVarIndexedPaths.GetValueOnGameThread().ParseIntoArray(Paths, TEXT(","), /*bCullEmpty=*/true);
+	for (FString& Path : Paths)
+	{
+		Path.TrimStartAndEndInline();
+	}
+	return Paths;
+}
+
+static bool IsPathInIndexedFolders(const FString& Path)
+{
+	for (const FString& Folder : GetIndexedFolderPaths())
+	{
+		if (Path.StartsWith(Folder))
+		{
+			return true;
+		}
+	}
+	return false;
+}
+
 FSemanticSearchModule& FSemanticSearchModule::Get()
 {
 	static const FName ModuleName(TEXT("SemanticSearch"));
 	return FModuleManager::Get().GetModuleChecked<FSemanticSearchModule>(ModuleName);
 }
 
+bool FSemanticSearchModule::IsInIndexedFolder(const FAssetData& AssetData)
+{
+	return IsPathInIndexedFolders(AssetData.PackageName.ToString());
+}
+
 void FSemanticSearchModule::StartupModule()
 {
 	UE_LOG(LogSemanticSearch, Log, TEXT("SemanticSearch module starting up"));
@@ -29,13 +74,141 @@ void FSemanticSearchModule::StartupModule()
 
 	// Default to a flat index. Dimension is set when first embedding arrives.
 	// Callers can replace this with a PQ index via SetVectorIndex().
-	constexpr int32 DefaultDimension = 768;
+	constexpr int32 DefaultDimension = 1536;
 	FAssetProcessorManager::Get().SetVectorIndex(MakeShared<FFlatVectorIndex>(DefaultDimension));
+
+	IAssetRegistry& AssetRegistry = FModuleManager::LoadModuleChecked<FAssetRegistryModule>(TEXT("AssetRegistry")).Get();
+	AssetAddedDelegateHandle = AssetRegistry.OnAssetAdded().AddRaw(
+		this, &FSemanticSearchModule::OnAssetAdded);
+	AssetUpdatedOnDiskDelegateHandle = AssetRegistry.OnAssetUpdatedOnDisk().AddRaw(
+		this, &FSemanticSearchModule::OnAssetUpdatedOnDisk);
+	AssetRenamedDelegateHandle = AssetRegistry.OnAssetRenamed().AddRaw(
+		this, &FSemanticSearchModule::OnAssetRenamed);
+	AssetRemovedDelegateHandle = AssetRegistry.OnAssetRemoved().AddRaw(
+		this, &FSemanticSearchModule::OnAssetRemoved);
 }
 
 void FSemanticSearchModule::ShutdownModule()
 {
 	UE_LOG(LogSemanticSearch, Log, TEXT("SemanticSearch module shutting down"));
+
+	IAssetRegistry& AssetRegistry = FModuleManager::LoadModuleChecked<FAssetRegistryModule>(TEXT("AssetRegistry")).Get();
+	AssetRegistry.OnAssetAdded().Remove(AssetAddedDelegateHandle);
+	AssetRegistry.OnAssetUpdatedOnDisk().Remove(AssetUpdatedOnDiskDelegateHandle);
+	AssetRegistry.OnAssetRenamed().Remove(AssetRenamedDelegateHandle);
+	AssetRegistry.OnAssetRemoved().Remove(AssetRemovedDelegateHandle);
+}
+
+void FSemanticSearchModule::TriggerIndexing(const FAssetData& AssetData, bool bBuildOnMiss)
+{
+	const int64 AssetID = GetAssetIndexID(AssetData);
+	FAssetProcessorManager::Get().GetEmbeddingData(AssetData,
+		[AssetID, AssetPath = AssetData.GetObjectPathString(), bBuildOnMiss](FAssetEmbeddingResult&& Result, FString&& Error) mutable
+		{
+			if (!Error.IsEmpty() || Result.Embedding.Num() == 0)
+			{
+				// "Not cached" on a cache-only query is expected — only warn when a full build was attempted
+				if (bBuildOnMiss)
+				{
+					UE_LOG(LogSemanticSearch, Warning, TEXT("Indexing failed for %s: %s"), *AssetPath, Error.IsEmpty() ? TEXT("empty embedding") : *Error);
+				}
+				return;
+			}
+			if (TSharedPtr<IVectorIndex> VectorIndex = FAssetProcessorManager::Get().GetVectorIndex())
+			{
+				// Remove first so resaves and any BuildDerivedData-added entries don't create duplicates
+				VectorIndex->Remove(AssetID);
+				VectorIndex->Add(AssetID, Result.Embedding);
+				UE_LOG(LogSemanticSearch, Log, TEXT("Indexed %s (%d dims, index size: %lld)"), *AssetPath, Result.Embedding.Num(), VectorIndex->GetCount());
+			}
+		}, bBuildOnMiss);
+}
+
+void FSemanticSearchModule::OnAssetAdded(const FAssetData& AssetData)
+{
+	if (IsInIndexedFolder(AssetData))
+	{
+		UE_LOG(LogSemanticSearch, Log, TEXT("OnAssetAdded: %s"), *AssetData.GetObjectPathString());
+		TriggerIndexing(AssetData, CVarAutoIndexOnStartup.GetValueOnGameThread());
+	}
+}
+
+void FSemanticSearchModule::OnAssetUpdatedOnDisk(const FAssetData& AssetData)
+{
+	if (!CVarAutoIndexOnSave.GetValueOnGameThread())
+	{
+		return;
+	}
+	if (IsInIndexedFolder(AssetData))
+	{
+		UE_LOG(LogSemanticSearch, Log, TEXT("OnAssetUpdatedOnDisk: %s"), *AssetData.GetObjectPathString());
+		TriggerIndexing(AssetData, /*bBuildOnMiss=*/true);
+	}
+}
+
+void FSemanticSearchModule::OnAssetRenamed(const FAssetData& NewAssetData, const FString& OldObjectPath)
+{
+	const bool bOldInFolder = IsPathInIndexedFolders(OldObjectPath);
+	const bool bNewInFolder = IsInIndexedFolder(NewAssetData);
+
+	UE_LOG(LogSemanticSearch, Log, TEXT("OnAssetRenamed: %s -> %s (OldInFolder=%d, NewInFolder=%d)"),
+		*OldObjectPath, *NewAssetData.GetObjectPathString(), bOldInFolder, bNewInFolder);
+
+	if (!bOldInFolder && !bNewInFolder)
+	{
+		UE_LOG(LogSemanticSearch, Log, TEXT("OnAssetRenamed: neither path in indexed folder, ignoring"));
+		return;
+	}
+
+	TSharedPtr<IVectorIndex> VectorIndex = FAssetProcessorManager::Get().GetVectorIndex();
+
+	const int64 OldID = GetAssetIndexID(OldObjectPath);
+	const int64 NewID = GetAssetIndexID(NewAssetData);
+
+	UE_LOG(LogSemanticSearch, Log, TEXT("OnAssetRenamed: OldID=%lld, NewID=%lld"), OldID, NewID);
+
+	// Try to reuse the existing embedding from the index to avoid a round-trip to the backend
+	TArray<float> Embedding;
+	if (bOldInFolder && VectorIndex && VectorIndex->TryGetEmbedding(OldID, Embedding))
+	{
+		VectorIndex->Remove(OldID);
+		if (bNewInFolder)
+		{
+			VectorIndex->Add(NewID, Embedding);
+			UE_LOG(LogSemanticSearch, Log, TEXT("OnAssetRenamed: reused embedding (%d dims) for %s -> %s (index size: %lld)"),
+				Embedding.Num(), *OldObjectPath, *NewAssetData.GetObjectPathString(), VectorIndex->GetCount());
+		}
+		else
+		{
+			UE_LOG(LogSemanticSearch, Log, TEXT("OnAssetRenamed: removed %s from index (moved out of indexed folder, index size: %lld)"),
+				*OldObjectPath, VectorIndex->GetCount());
+		}
+	}
+	else
+	{
+		// Embedding not in index (e.g. index was rebuilt without this asset) — fall back to full reindex
+		if (bOldInFolder && VectorIndex)
+		{
+			const bool bRemoved = VectorIndex->Remove(OldID);
+			UE_LOG(LogSemanticSearch, Log, TEXT("OnAssetRenamed: embedding not in index for %s (Remove=%d)"), *OldObjectPath, bRemoved);
+		}
+		if (bNewInFolder)
+		{
+			UE_LOG(LogSemanticSearch, Log, TEXT("OnAssetRenamed: triggering full reindex for %s"), *NewAssetData.GetObjectPathString());
+			TriggerIndexing(NewAssetData, CVarAutoIndexOnSave.GetValueOnGameThread());
+		}
+	}
+}
+
+void FSemanticSearchModule::OnAssetRemoved(const FAssetData& AssetData)
+{
+	if (IsInIndexedFolder(AssetData))
+	{
+		if (TSharedPtr<IVectorIndex> VectorIndex = FAssetProcessorManager::Get().GetVectorIndex())
+		{
+			VectorIndex->Remove(GetAssetIndexID(AssetData));
+		}
+	}
 }
 
 TSharedPtr<IEmbeddingProvider> FSemanticSearchModule::GetRegisteredEmbeddingProvider()
@@ -53,5 +226,9 @@ void FSemanticSearchModule::SetVectorIndex(TSharedPtr<IVectorIndex> InIndex)
 	FAssetProcessorManager::Get().SetVectorIndex(MoveTemp(InIndex));
 }
 
+void FSemanticSearchModule::IndexAsset(const FAssetData& AssetData)
+{
+	TriggerIndexing(AssetData, /*bBuildOnMiss=*/true);
 }
 
+}
diff --git a/Engine/Plugins/Experimental/SemanticSearch/Source/SemanticSearch/Private/SemanticSearchModule.h b/Engine/Plugins/Experimental/SemanticSearch/Source/SemanticSearch/Private/SemanticSearchModule.h
--- a/Engine/Plugins/Experimental/SemanticSearch/Source/SemanticSearch/Private/SemanticSearchModule.h
+++ b/Engine/Plugins/Experimental/SemanticSearch/Source/SemanticSearch/Private/SemanticSearchModule.h
@@ -7,6 +7,7 @@
 #include "CoreMinimal.h"
 #include "Modules/ModuleInterface.h"
 #include "Templates/SharedPointer.h"
+#include "AssetRegistry/AssetData.h"
 
 DECLARE_LOG_CATEGORY_EXTERN(LogSemanticSearch, Log, All);
 
@@ -22,11 +23,23 @@ public:
 	virtual void StartupModule() override;
 	virtual void ShutdownModule() override;
 
-
 	virtual void RegisterEmbeddingProvider(TUniquePtr<IEmbeddingProvider>&& EmbeddingProvider) override;
 	virtual void SetVectorIndex(TSharedPtr<IVectorIndex> InIndex) override;
+	virtual void IndexAsset(const FAssetData& AssetData) override;
 
 private:
+	static bool IsInIndexedFolder(const FAssetData& AssetData);
+	void TriggerIndexing(const FAssetData& AssetData, bool bBuildOnMiss);
+
+	void OnAssetAdded(const FAssetData& AssetData);
+	void OnAssetUpdatedOnDisk(const FAssetData& AssetData);
+	void OnAssetRenamed(const FAssetData& NewAssetData, const FString& OldObjectPath);
+	void OnAssetRemoved(const FAssetData& AssetData);
+
 	TSharedPtr<IEmbeddingProvider> RegisteredEmbeddingProvider;
+	FDelegateHandle AssetAddedDelegateHandle;
+	FDelegateHandle AssetUpdatedOnDiskDelegateHandle;
+	FDelegateHandle AssetRenamedDelegateHandle;
+	FDelegateHandle AssetRemovedDelegateHandle;
 };
-}
\ No newline at end of file
+}
diff --git a/Engine/Plugins/Experimental/SemanticSearch/Source/SemanticSearch/Private/SemanticSearchTestCommands.cpp b/Engine/Plugins/Experimental/SemanticSearch/Source/SemanticSearch/Private/SemanticSearchTestCommands.cpp
--- a/Engine/Plugins/Experimental/SemanticSearch/Source/SemanticSearch/Private/SemanticSearchTestCommands.cpp
+++ b/Engine/Plugins/Experimental/SemanticSearch/Source/SemanticSearch/Private/SemanticSearchTestCommands.cpp
@@ -1,16 +1,46 @@
 // Copyright Epic Games, Inc. All Rights Reserved.
 
 #include "SemanticSearchModule.h"
+#include "AssetProcessorManager.h"
+#include "AssetProcessors/AssetProcessorUtils.h"
+#include "ISemanticSearchModule.h"
 #include "RemoteEmbeddingProvider.h"
 #include "HAL/ConsoleManager.h"
 #include "Misc/FileHelper.h"
 #include "Misc/Paths.h"
+#include "AssetRegistry/AssetRegistryModule.h"
 
 namespace UE::SemanticSearch::Private
 {
 
 static TSharedPtr<FRemoteEmbeddingProvider> GTestProvider;
 
+// Returns all assets in the folders listed in SemanticSearch.IndexedPaths.
+static TArray<FAssetData> GetAssetsInIndexedFolders(IAssetRegistry& AssetRegistry)
+{
+	TArray<FString> IndexedPaths;
+	if (IConsoleVariable* CVar = IConsoleManager::Get().FindConsoleVariable(TEXT("SemanticSearch.IndexedPaths")))
+	{
+		CVar->GetString().ParseIntoArray(IndexedPaths, TEXT(","), /*bCullEmpty=*/true);
+	}
+
+	if (IndexedPaths.IsEmpty())
+	{
+		UE_LOG(LogSemanticSearch, Warning, TEXT("GetAssetsInIndexedFolders: SemanticSearch.IndexedPaths is empty or CVar not found"));
+	}
+
+	TArray<FAssetData> Result;
+	for (FString& Folder : IndexedPaths)
+	{
+		Folder.TrimStartAndEndInline();
+		FString FolderPath = Folder.EndsWith(TEXT("/")) ? Folder.LeftChop(1) : Folder;
+		const int32 Before = Result.Num();
+		AssetRegistry.GetAssetsByPath(FName(FolderPath), Result, /*bRecursive=*/true);
+		UE_LOG(LogSemanticSearch, Log, TEXT("GetAssetsInIndexedFolders: '%s' -> %d assets"), *FolderPath, Result.Num() - Before);
+	}
+	return Result;
+}
+
 static TSharedPtr<FRemoteEmbeddingProvider> GetOrCreateProvider()
 {
 	if (!GTestProvider.IsValid())
@@ -20,6 +50,141 @@ static TSharedPtr<FRemoteEmbeddingProvider> GetOrCreateProvider()
 	return GTestProvider;
 }
 
+// SemanticSearch.IndexStats
+static FAutoConsoleCommand IndexStatsCommand(
+	TEXT("SemanticSearch.IndexStats"),
+	TEXT("Print the number of assets currently in the vector index."),
+	FConsoleCommandDelegate::CreateLambda([]()
+	{
+		TSharedPtr<IVectorIndex> VectorIndex = FAssetProcessorManager::Get().GetVectorIndex();
+		if (!VectorIndex)
+		{
+			UE_LOG(LogSemanticSearch, Warning, TEXT("No vector index is registered."));
+			return;
+		}
+
+		IAssetRegistry& AssetRegistry = FModuleManager::LoadModuleChecked<FAssetRegistryModule>(TEXT("AssetRegistry")).Get();
+		TArray<FAssetData> Assets = GetAssetsInIndexedFolders(AssetRegistry);
+
+		int32 Indexed = 0;
+		for (const FAssetData& Asset : Assets)
+		{
+			if (VectorIndex->Contains(GetAssetIndexID(Asset)))
+			{
+				++Indexed;
+			}
+		}
+		const int32 Total = Assets.Num();
+		UE_LOG(LogSemanticSearch, Log, TEXT("Index: %lld entries | Folder assets: %d total, %d indexed, %d not indexed (%d dims)"),
+			VectorIndex->GetCount(), Total, Indexed, Total - Indexed, VectorIndex->GetDimension());
+	})
+);
+
+// SemanticSearch.IndexAll
+static FAutoConsoleCommand IndexAllCommand(
+	TEXT("SemanticSearch.IndexAll"),
+	TEXT("Index all assets in the configured folders that are not yet in the FAISS index."),
+	FConsoleCommandDelegate::CreateLambda([]()
+	{
+		TSharedPtr<IVectorIndex> VectorIndex = FAssetProcessorManager::Get().GetVectorIndex();
+		if (!VectorIndex)
+		{
+			UE_LOG(LogSemanticSearch, Warning, TEXT("No vector index is registered."));
+			return;
+		}
+
+		IAssetRegistry& AssetRegistry = FModuleManager::LoadModuleChecked<FAssetRegistryModule>(TEXT("AssetRegistry")).Get();
+		TArray<FAssetData> Assets = GetAssetsInIndexedFolders(AssetRegistry);
+		UE_LOG(LogSemanticSearch, Log, TEXT("IndexAll: found %d assets in indexed folders"), Assets.Num());
+
+		int32 Queued = 0;
+		int32 AlreadyIndexed = 0;
+		for (const FAssetData& Asset : Assets)
+		{
+			if (VectorIndex->Contains(GetAssetIndexID(Asset)))
+			{
+				++AlreadyIndexed;
+			}
+			else
+			{
+				ISemanticSearchModule::Get().IndexAsset(Asset);
+				++Queued;
+			}
+		}
+		UE_LOG(LogSemanticSearch, Log, TEXT("IndexAll: queued %d assets for indexing, %d already indexed"),
+			Queued, AlreadyIndexed);
+	})
+);
+
+// SemanticSearch.TestSearch <query text>
+static FAutoConsoleCommand TestSearchCommand(
+	TEXT("SemanticSearch.TestSearch"),
+	TEXT("Search the vector index with a text query. Usage: SemanticSearch.TestSearch <query text>"),
+	FConsoleCommandWithArgsDelegate::CreateLambda([](const TArray<FString>& Args)
+	{
+		if (Args.Num() < 1)
+		{
+			UE_LOG(LogSemanticSearch, Warning, TEXT("Usage: SemanticSearch.TestSearch <query text>"));
+			return;
+		}
+
+		TSharedPtr<IVectorIndex> VectorIndex = FAssetProcessorManager::Get().GetVectorIndex();
+		if (!VectorIndex || VectorIndex->GetCount() == 0)
+		{
+			UE_LOG(LogSemanticSearch, Warning, TEXT("Vector index is empty — no assets indexed yet."));
+			return;
+		}
+
+		const FString QueryText = FString::Join(Args, TEXT(" "));
+		UE_LOG(LogSemanticSearch, Log, TEXT("Searching for: \"%s\" (index size: %lld)"), *QueryText, VectorIndex->GetCount());
+
+		// Build ID -> asset path map from all configured indexed folders
+		TMap<int64, FString> IDToPath;
+		IAssetRegistry& AssetRegistry = FModuleManager::LoadModuleChecked<FAssetRegistryModule>(TEXT("AssetRegistry")).Get();
+		TArray<FString> IndexedPaths;
+		if (IConsoleVariable* CVar = IConsoleManager::Get().FindConsoleVariable(TEXT("SemanticSearch.IndexedPaths")))
+		{
+			CVar->GetString().ParseIntoArray(IndexedPaths, TEXT(","), /*bCullEmpty=*/true);
+		}
+		for (FString& Folder : IndexedPaths)
+		{
+			Folder.TrimStartAndEndInline();
+			// GetAssetsByPath expects the path without trailing slash
+			TArray<FAssetData> Assets;
+			FString FolderPath = Folder.EndsWith(TEXT("/")) ? Folder.LeftChop(1) : Folder;
+			AssetRegistry.GetAssetsByPath(FName(FolderPath), Assets, /*bRecursive=*/true);
+			for (const FAssetData& AssetData : Assets)
+			{
+				IDToPath.Add(GetAssetIndexID(AssetData), AssetData.GetObjectPathString());
+			}
+		}
+
+		TSharedPtr<FRemoteEmbeddingProvider> Provider = GetOrCreateProvider();
+		Provider->GenerateEmbeddingAsync(QueryText,
+			[VectorIndex, IDToPath = MoveTemp(IDToPath)](const FEmbeddingResponse& Response)
+			{
+				if (!Response.ErrorMessage.IsEmpty())
+				{
+					UE_LOG(LogSemanticSearch, Error, TEXT("Embedding failed: %s"), *Response.ErrorMessage);
+					return;
+				}
+
+				constexpr int32 K = 5;
+				TArray<FSearchResult> Results = VectorIndex->Search(Response.Embedding, K);
+
+				UE_LOG(LogSemanticSearch, Log, TEXT("Top %d results:"), Results.Num());
+				for (int32 i = 0; i < Results.Num(); ++i)
+				{
+					const FString* Path = IDToPath.Find(Results[i].ID);
+					UE_LOG(LogSemanticSearch, Log, TEXT("  [%d] dist=%.4f  %s"),
+						i + 1,
+						Results[i].Distance,
+						Path ? **Path : TEXT("<unknown>"));
+				}
+			}, nullptr);
+	})
+);
+
 // SemanticSearch.TestCaption <image_file_path>
 static FAutoConsoleCommand TestCaptionCommand(
 	TEXT("SemanticSearch.TestCaption"),
diff --git a/Engine/Plugins/Experimental/SemanticSearch/Source/SemanticSearch/Private/VectorIndex/FlatVectorIndex.cpp b/Engine/Plugins/Experimental/SemanticSearch/Source/SemanticSearch/Private/VectorIndex/FlatVectorIndex.cpp
--- a/Engine/Plugins/Experimental/SemanticSearch/Source/SemanticSearch/Private/VectorIndex/FlatVectorIndex.cpp
+++ b/Engine/Plugins/Experimental/SemanticSearch/Source/SemanticSearch/Private/VectorIndex/FlatVectorIndex.cpp
@@ -16,10 +16,10 @@ FFlatVectorIndex::FFlatVectorIndex(int32 InDimension)
 	: Dimension(InDimension)
 {
 	faiss::IndexFlatL2* FlatIndex = new faiss::IndexFlatL2(Dimension);
-	IndexWithIDs = new faiss::IndexIDMap(FlatIndex);
+	IndexWithIDs = new faiss::IndexIDMap2(FlatIndex);
 }
 
-FFlatVectorIndex::FFlatVectorIndex(int32 InDimension, faiss::IndexIDMap* InIndex)
+FFlatVectorIndex::FFlatVectorIndex(int32 InDimension, faiss::IndexIDMap2* InIndex)
 	: Dimension(InDimension)
 	, IndexWithIDs(InIndex)
 {
@@ -104,6 +104,25 @@ int64 FFlatVectorIndex::GetCount() const
 	return IndexWithIDs->ntotal;
 }
 
+bool FFlatVectorIndex::Contains(int64 ID) const
+{
+	FReadScopeLock ReadLock(IndexLock);
+	return IndexWithIDs->rev_map.find(ID) != IndexWithIDs->rev_map.end();
+}
+
+bool FFlatVectorIndex::TryGetEmbedding(int64 ID, TArray<float>& OutEmbedding) const
+{
+	FReadScopeLock ReadLock(IndexLock);
+	if (IndexWithIDs->rev_map.find(ID) == IndexWithIDs->rev_map.end())
+	{
+		return false;
+	}
+
+	OutEmbedding.SetNumUninitialized(Dimension);
+	IndexWithIDs->reconstruct(static_cast<faiss::idx_t>(ID), OutEmbedding.GetData());
+	return true;
+}
+
 TArray<uint8> FFlatVectorIndex::Serialize() const
 {
 	FReadScopeLock ReadLock(IndexLock);
@@ -129,7 +148,7 @@ TUniquePtr<FFlatVectorIndex> FFlatVectorIndex::Deserialize(TConstArrayView<uint8
 		return nullptr;
 	}
 
-	faiss::IndexIDMap* IDMapIndex = static_cast<faiss::IndexIDMap*>(LoadedIndex);
+	faiss::IndexIDMap2* IDMapIndex = static_cast<faiss::IndexIDMap2*>(LoadedIndex);
 	return TUniquePtr<FFlatVectorIndex>(new FFlatVectorIndex(ExpectedDimension, IDMapIndex));
 }
 
diff --git a/Engine/Plugins/Experimental/SemanticSearch/Source/SemanticSearch/Private/VectorIndex/FlatVectorIndex.h b/Engine/Plugins/Experimental/SemanticSearch/Source/SemanticSearch/Private/VectorIndex/FlatVectorIndex.h
--- a/Engine/Plugins/Experimental/SemanticSearch/Source/SemanticSearch/Private/VectorIndex/FlatVectorIndex.h
+++ b/Engine/Plugins/Experimental/SemanticSearch/Source/SemanticSearch/Private/VectorIndex/FlatVectorIndex.h
@@ -2,8 +2,8 @@
 
 #pragma once
 
-#include "Interfaces/IVectorIndex.h"
 #include "HAL/CriticalSection.h"
+#include "Interfaces/IVectorIndex.h"
 #include "Templates/UniquePtr.h"
 
 THIRD_PARTY_INCLUDES_START
@@ -38,15 +38,17 @@ public:
 	virtual int64 GetCount() const override;
 	virtual int32 GetDimension() const override { return Dimension; }
 	virtual TArray<uint8> Serialize() const override;
+	virtual bool Contains(int64 ID) const override;
+	virtual bool TryGetEmbedding(int64 ID, TArray<float>& OutEmbedding) const override;
 
 	/** Deserialize a full flat index. */
 	static TUniquePtr<FFlatVectorIndex> Deserialize(TConstArrayView<uint8> Data, int32 ExpectedDimension);
 
 private:
-	FFlatVectorIndex(int32 InDimension, faiss::IndexIDMap* InIndex);
+	FFlatVectorIndex(int32 InDimension, faiss::IndexIDMap2* InIndex);
 
 	int32 Dimension;
-	faiss::IndexIDMap* IndexWithIDs = nullptr;
+	faiss::IndexIDMap2* IndexWithIDs = nullptr;
 	mutable FRWLock IndexLock;
 };
 
diff --git a/Engine/Plugins/Experimental/SemanticSearch/Source/SemanticSearch/Public/AssetProcessorManager.h b/Engine/Plugins/Experimental/SemanticSearch/Source/SemanticSearch/Public/AssetProcessorManager.h
--- a/Engine/Plugins/Experimental/SemanticSearch/Source/SemanticSearch/Public/AssetProcessorManager.h
+++ b/Engine/Plugins/Experimental/SemanticSearch/Source/SemanticSearch/Public/AssetProcessorManager.h
@@ -43,7 +43,7 @@ public:
 
 	SEMANTICSEARCH_API bool StarCachingAsset(FAssetData&& InAssetData);
 	SEMANTICSEARCH_API bool GetCaptionData(const FAssetData& InAssetData, TUniqueFunction<void(FAssetCaptionResult&&, FString&&)>&& InOnDataAvailable);
-	SEMANTICSEARCH_API bool GetEmbeddingData(const FAssetData& InAssetData, TUniqueFunction<void(FAssetEmbeddingResult&&, FString&&)>&& InOnDataAvailable);
+	SEMANTICSEARCH_API bool GetEmbeddingData(const FAssetData& InAssetData, TUniqueFunction<void(FAssetEmbeddingResult&&, FString&&)>&& InOnDataAvailable, bool bBuildOnMiss = true);
 
 	SEMANTICSEARCH_API void SetVectorIndex(TSharedPtr<IVectorIndex> InIndex);
 	SEMANTICSEARCH_API TSharedPtr<IVectorIndex> GetVectorIndex() const;
diff --git a/Engine/Plugins/Experimental/SemanticSearch/Source/SemanticSearch/Public/ISemanticSearchModule.h b/Engine/Plugins/Experimental/SemanticSearch/Source/SemanticSearch/Public/ISemanticSearchModule.h
--- a/Engine/Plugins/Experimental/SemanticSearch/Source/SemanticSearch/Public/ISemanticSearchModule.h
+++ b/Engine/Plugins/Experimental/SemanticSearch/Source/SemanticSearch/Public/ISemanticSearchModule.h
@@ -2,6 +2,7 @@
 
 #pragma once
 
+#include "AssetRegistry/AssetData.h"
 #include "Interfaces/IEmbeddingProvider.h"
 #include "Interfaces/IVectorIndex.h"
 #include "Modules/ModuleInterface.h"
@@ -12,7 +13,7 @@
 namespace UE::SemanticSearch
 {
 	class ISemanticSearchModule : public IModuleInterface
-	{
+	{ 
 	public:
 		static ISemanticSearchModule& Get()
 		{
@@ -24,5 +25,8 @@ namespace UE::SemanticSearch
 
 		/** Set the vector index used for similarity search. Replaces any existing index. */
 		virtual void SetVectorIndex(TSharedPtr<IVectorIndex> InIndex) = 0;
+
+		/** Index a single asset, always building from the backend if not already cached. */
+		virtual void IndexAsset(const FAssetData& AssetData) = 0;
 	};
 }
\ No newline at end of file
diff --git a/Engine/Plugins/Experimental/SemanticSearch/Source/SemanticSearch/Public/Interfaces/IVectorIndex.h b/Engine/Plugins/Experimental/SemanticSearch/Source/SemanticSearch/Public/Interfaces/IVectorIndex.h
--- a/Engine/Plugins/Experimental/SemanticSearch/Source/SemanticSearch/Public/Interfaces/IVectorIndex.h
+++ b/Engine/Plugins/Experimental/SemanticSearch/Source/SemanticSearch/Public/Interfaces/IVectorIndex.h
@@ -76,6 +76,11 @@ public:
 		return {};
 	}
 
+	TArray<uint8> Quantize(TConstArrayView<float> Vector) const
+	{
+		return Quantize(Vector, 1);
+	}
+
 	/**
 	 * Add pre-quantized vectors directly, skipping encoding.
 	 * check(SupportsQuantization()) — crashes if not supported.
@@ -104,10 +109,12 @@ public:
 		return GetCount() < CountBefore;
 	}
 
-	TArray<uint8> Quantize(TConstArrayView<float> Vector) const
-	{
-		return Quantize(Vector, 1);
-	}
+
+	/** Returns true if a vector with the given ID exists in the index. */
+	virtual bool Contains(int64 ID) const { return false; }
+
+	/** Retrieve the embedding for a given ID. Returns false if not supported or ID not found. */
+	virtual bool TryGetEmbedding(int64 ID, TArray<float>& OutEmbedding) const { return false; }
 };
 
 } // namespace UE::SemanticSearch
