Fixed a crash in the semantic search module when closing the editor.

#Synced-CL 52301593
#Synced-CL 52301593

diff --git a/Engine/Plugins/Experimental/SemanticSearch/Source/SemanticSearch/Private/SemanticSearchModule.cpp b/Engine/Plugins/Experimental/SemanticSearch/Source/SemanticSearch/Private/SemanticSearchModule.cpp
--- a/Engine/Plugins/Experimental/SemanticSearch/Source/SemanticSearch/Private/SemanticSearchModule.cpp
+++ b/Engine/Plugins/Experimental/SemanticSearch/Source/SemanticSearch/Private/SemanticSearchModule.cpp
@@ -92,11 +92,14 @@ void FSemanticSearchModule::ShutdownModule()
 {
 	UE_LOG(LogSemanticSearch, Log, TEXT("SemanticSearch module shutting down"));
 
-	IAssetRegistry& AssetRegistry = FModuleManager::LoadModuleChecked<FAssetRegistryModule>(TEXT("AssetRegistry")).Get();
-	AssetRegistry.OnAssetAdded().Remove(AssetAddedDelegateHandle);
-	AssetRegistry.OnAssetUpdatedOnDisk().Remove(AssetUpdatedOnDiskDelegateHandle);
-	AssetRegistry.OnAssetRenamed().Remove(AssetRenamedDelegateHandle);
-	AssetRegistry.OnAssetRemoved().Remove(AssetRemovedDelegateHandle);
+	if (FAssetRegistryModule* AssetRegistryModule = FModuleManager::GetModulePtr<FAssetRegistryModule>(TEXT("AssetRegistry")))
+	{
+		IAssetRegistry& AssetRegistry = AssetRegistryModule->Get();
+		AssetRegistry.OnAssetAdded().Remove(AssetAddedDelegateHandle);
+		AssetRegistry.OnAssetUpdatedOnDisk().Remove(AssetUpdatedOnDiskDelegateHandle);
+		AssetRegistry.OnAssetRenamed().Remove(AssetRenamedDelegateHandle);
+		AssetRegistry.OnAssetRemoved().Remove(AssetRemovedDelegateHandle);
+	}
 }
 
 void FSemanticSearchModule::TriggerIndexing(const FAssetData& AssetData, bool bBuildOnMiss)
