Add experimental semantic search integration to the Content Browser
Introduce a new editor integrations module for the SemanticSearch plugin that hooks into the Content Browser via a new IAssetSearchOverride extension point
Add per-view override instances that intercept the asset filtering pipeline and redirect search text handling when active, and expose a toggle widget in the search box via a new ExtraWidget slot on SAssetSearchBox and SFilterSearchBox
Currently the semantic search implementation is a muck up to test the extension in the CB.

#Synced-CL 52301593

diff --git a/Engine/Plugins/Experimental/SemanticSearch/SemanticSearch.uplugin b/Engine/Plugins/Experimental/SemanticSearch/SemanticSearch.uplugin
--- a/Engine/Plugins/Experimental/SemanticSearch/SemanticSearch.uplugin
+++ b/Engine/Plugins/Experimental/SemanticSearch/SemanticSearch.uplugin
@@ -21,6 +21,12 @@
 			"Type": "Editor",
 			"LoadingPhase": "Default",
 			"PlatformAllowList": [ "Win64", "Mac" ]
+		},
+		{
+			"Name": "SemanticSearchEditorIntegrations",
+			"Type": "Editor",
+			"LoadingPhase": "Default",
+			"PlatformAllowList": [ "Win64", "Mac" ]
 		}
 	]
 }
\ No newline at end of file
diff --git a/Engine/Plugins/Experimental/SemanticSearch/Source/SemanticSearch/Private/SemanticSearchDerivedDataBuilder.cpp b/Engine/Plugins/Experimental/SemanticSearch/Source/SemanticSearch/Private/SemanticSearchDerivedDataBuilder.cpp
--- a/Engine/Plugins/Experimental/SemanticSearch/Source/SemanticSearch/Private/SemanticSearchDerivedDataBuilder.cpp
+++ b/Engine/Plugins/Experimental/SemanticSearch/Source/SemanticSearch/Private/SemanticSearchDerivedDataBuilder.cpp
@@ -260,7 +260,7 @@ void FSemanticSearchBuildCacheTask::BuildDerivedData()
 			TRACE_CPUPROFILER_EVENT_SCOPE_STR("SemanticSearch::BuildDerivedData::CaptionRequestReady");
 			if (bHasGeneratedRequest)
 			{
-				if (TSharedPtr<IEmbeddingProvider> EmbeddingProvider = FSemanticSearchModule::Get().GetRegisteredEmbeddingProvider())
+				if (TSharedPtr<IEmbeddingProvider> EmbeddingProvider = FSemanticSearchModule::Get().GetEmbeddingProvider())
 				{
 					EmbeddingProvider->GenerateCaptionAsync(MoveTemp(CaptionRequest), [this, EmbeddingProvider](FCaptionResponse&& Caption)
 						{
diff --git a/Engine/Plugins/Experimental/SemanticSearch/Source/SemanticSearch/Private/SemanticSearchModule.cpp b/Engine/Plugins/Experimental/SemanticSearch/Source/SemanticSearch/Private/SemanticSearchModule.cpp
--- a/Engine/Plugins/Experimental/SemanticSearch/Source/SemanticSearch/Private/SemanticSearchModule.cpp
+++ b/Engine/Plugins/Experimental/SemanticSearch/Source/SemanticSearch/Private/SemanticSearchModule.cpp
@@ -214,7 +214,7 @@ void FSemanticSearchModule::OnAssetRemoved(const FAssetData& AssetData)
 	}
 }
 
-TSharedPtr<IEmbeddingProvider> FSemanticSearchModule::GetRegisteredEmbeddingProvider()
+TSharedPtr<IEmbeddingProvider> FSemanticSearchModule::GetEmbeddingProvider() const
 {
 	return RegisteredEmbeddingProvider;
 }
diff --git a/Engine/Plugins/Experimental/SemanticSearch/Source/SemanticSearch/Private/SemanticSearchModule.h b/Engine/Plugins/Experimental/SemanticSearch/Source/SemanticSearch/Private/SemanticSearchModule.h
--- a/Engine/Plugins/Experimental/SemanticSearch/Source/SemanticSearch/Private/SemanticSearchModule.h
+++ b/Engine/Plugins/Experimental/SemanticSearch/Source/SemanticSearch/Private/SemanticSearchModule.h
@@ -18,7 +18,7 @@ class FSemanticSearchModule : public ISemanticSearchModule
 {
 public:
 	static FSemanticSearchModule& Get();
-	virtual TSharedPtr<IEmbeddingProvider> GetRegisteredEmbeddingProvider();
+	virtual TSharedPtr<IEmbeddingProvider> GetEmbeddingProvider() const override;
 
 	virtual void StartupModule() override;
 	virtual void ShutdownModule() override;
diff --git a/Engine/Plugins/Experimental/SemanticSearch/Source/SemanticSearch/Public/ISemanticSearchModule.h b/Engine/Plugins/Experimental/SemanticSearch/Source/SemanticSearch/Public/ISemanticSearchModule.h
--- a/Engine/Plugins/Experimental/SemanticSearch/Source/SemanticSearch/Public/ISemanticSearchModule.h
+++ b/Engine/Plugins/Experimental/SemanticSearch/Source/SemanticSearch/Public/ISemanticSearchModule.h
@@ -23,6 +23,9 @@ namespace UE::SemanticSearch
 
 		virtual void RegisterEmbeddingProvider(TUniquePtr<IEmbeddingProvider>&& EmbeddingProvider) = 0;
 
+		/** Returns the active embedding provider, or nullptr if none has been registered. */
+		virtual TSharedPtr<IEmbeddingProvider> GetEmbeddingProvider() const = 0;
+
 		/** Set the vector index used for similarity search. Replaces any existing index. */
 		virtual void SetVectorIndex(TSharedPtr<IVectorIndex> InIndex) = 0;
 
diff --git a/Engine/Plugins/Experimental/SemanticSearch/Source/SemanticSearch/SemanticSearch.Build.cs b/Engine/Plugins/Experimental/SemanticSearch/Source/SemanticSearch/SemanticSearch.Build.cs
--- a/Engine/Plugins/Experimental/SemanticSearch/Source/SemanticSearch/SemanticSearch.Build.cs
+++ b/Engine/Plugins/Experimental/SemanticSearch/Source/SemanticSearch/SemanticSearch.Build.cs
@@ -18,7 +18,6 @@ public class SemanticSearch : ModuleRules
 				"AssetRegistry",
 				"Core",
 				"CoreUObject",
-				"DerivedDataCache",
 				"Engine",
 				"HTTP",
 				"ImageCore",
@@ -28,5 +27,10 @@ public class SemanticSearch : ModuleRules
 				"FAISS"
 			});
 
+		PublicDependencyModuleNames.AddRange(
+			new string[]
+			{
+				"DerivedDataCache",
+			});
 	}
 }
\ No newline at end of file
diff --git a/Engine/Plugins/Experimental/SemanticSearch/Source/SemanticSearchEditorIntegrations/Private/ContentBrowser/SemanticSearchCB.cpp b/Engine/Plugins/Experimental/SemanticSearch/Source/SemanticSearchEditorIntegrations/Private/ContentBrowser/SemanticSearchCB.cpp
new file mode 100644
--- /dev/null
+++ b/Engine/Plugins/Experimental/SemanticSearch/Source/SemanticSearchEditorIntegrations/Private/ContentBrowser/SemanticSearchCB.cpp
@@ -0,0 +1,369 @@
+// Copyright Epic Games, Inc. All Rights Reserved.
+
+#include "ContentBrowser/SemanticSearchCB.h"
+
+#include "AssetProcessorManager.h"
+#include "AssetViewTypes.h"
+#include "Async/Async.h"
+#include "Async/TaskGraphInterfaces.h"
+#include "ContentBrowserModule.h"
+#include "Experimental/ContentBrowserSearchOverride.h"
+#include "ISemanticSearchModule.h"
+#include "Input/Reply.h"
+#include "Interfaces/IEmbeddingProvider.h"
+#include "Math/UnrealMathUtility.h"
+#include "Math/VectorRegister.h"
+#include "Modules/ModuleManager.h"
+#include "Styling/AppStyle.h"
+#include "Styling/SlateColor.h"
+#include "Templates/SharedPointer.h"
+#include "Templates/UniquePtr.h"
+#include "Widgets/Input/SButton.h"
+#include "Widgets/Text/STextBlock.h"
+
+DEFINE_LOG_CATEGORY_STATIC(LogSemanticSearchCB, Log, All);
+
+namespace UE::SemanticSearch::ContentBrowser
+{
+namespace Private
+{
+	using namespace UE::Editor::ContentBrowser::Extension;
+	using namespace UE::SemanticSearch;
+
+	static float CosineSimilarity(const TArray<float>& A, const TArray<float>& B)
+	{
+		if (A.Num() != B.Num() || A.IsEmpty())
+		{
+			return 0.f;
+		}
+
+		const int32 Num = A.Num();
+		const float* __restrict DataA = A.GetData();
+		const float* __restrict DataB = B.GetData();
+
+		// Process 4 floats per iteration using SIMD fused multiply-add.
+		VectorRegister4Float VecDot  = VectorZeroFloat();
+		VectorRegister4Float VecNormA = VectorZeroFloat();
+		VectorRegister4Float VecNormB = VectorZeroFloat();
+
+		const int32 NumVec = Num & ~3;
+		for (int32 i = 0; i < NumVec; i += 4)
+		{
+			const VectorRegister4Float VA = VectorLoad(DataA + i);
+			const VectorRegister4Float VB = VectorLoad(DataB + i);
+			VecDot   = VectorMultiplyAdd(VA, VB, VecDot);
+			VecNormA = VectorMultiplyAdd(VA, VA, VecNormA);
+			VecNormB = VectorMultiplyAdd(VB, VB, VecNormB);
+		}
+
+		// Horizontal sum: add pairs, then add the two pair results.
+		auto HSum = [](VectorRegister4Float V) -> float
+		{
+			VectorRegister4Float Shuf = VectorSwizzle(V, 2, 3, 0, 1);
+			VectorRegister4Float Sum  = VectorAdd(V, Shuf);
+			Shuf = VectorSwizzle(Sum, 1, 0, 3, 2);
+			Sum  = VectorAdd(Sum, Shuf);
+			return VectorGetComponent(Sum, 0);
+		};
+
+		float Dot  = HSum(VecDot);
+		float NormA = HSum(VecNormA);
+		float NormB = HSum(VecNormB);
+
+		// Scalar tail for dimensions not divisible by 4.
+		for (int32 i = NumVec; i < Num; ++i)
+		{
+			Dot   += DataA[i] * DataB[i];
+			NormA += DataA[i] * DataA[i];
+			NormB += DataB[i] * DataB[i];
+		}
+
+		const float Denom = FMath::Sqrt(NormA) * FMath::Sqrt(NormB);
+		return Denom > UE_SMALL_NUMBER ? Dot / Denom : 0.f;
+	}
+
+	// ---------------------------------------------------------------------------
+
+	class FSemanticSearchAssetViewOverride : public ICBSearchOverrideInstance, public TSharedFromThis<FSemanticSearchAssetViewOverride>
+	{
+	public:
+		virtual TSharedRef<SWidget> GetOverrideModeToggleWidget() override;
+
+		virtual bool HasPendingResults() const override { return bHasInFlightRequest; }
+
+		virtual void OnSearchTextChanged(const FText& NewText) override {}
+
+		virtual void OnSearchTextCommitted(const FText& NewText, ETextCommit::Type) override
+		{
+			if (NewText.EqualTo(ActiveQuery))
+			{
+				return;
+			}
+
+			RequestQuickFrontendListRefresh();
+			ActiveQuery = NewText;
+		}
+
+		virtual void OnItemsAvailable(TArrayView<const TSharedPtr<FAssetViewItem>> Items) override
+		{
+			CurrentItems.Append(Items);
+
+			if (Items.IsEmpty())
+			{
+				return;
+			}
+
+			if (ActiveQuery.IsEmpty())
+			{
+				PublishPassThrough();
+			}
+			else if (bHasActiveResults)
+			{
+				PublishMatchingItems();
+			}
+		}
+
+		virtual void SortItemList(TArrayView<TSharedPtr<FAssetViewItem>> Items) override
+		{
+
+		}
+
+	protected:
+
+		// Incrementing the epoch silently invalidates all in-flight callbacks:
+		// every async job captures the epoch at launch and bails if it has changed.
+		// All epoch reads/writes happen inside AsyncTask(GameThread), so the game
+		// thread is the sole synchronization point — no atomics required.
+		virtual void OnFilteringResetImplementation() override
+		{
+			++ActiveEpoch;
+			bHasInFlightRequest = false;
+			CurrentItems.Reset();
+			bHasActiveResults = false;
+
+			if (ActiveQuery.IsEmpty())
+			{
+				SetUserSearching(false);
+				bHasActiveResults = false;
+				return;
+			}
+			else
+			{
+				SetUserSearching(true);
+				SubmitQueryAsync(ActiveQuery.ToString());
+			}
+		}
+
+	private:
+
+		// Bumps the epoch (invalidating any previous query's callbacks) and launches
+		// an async embedding request for the query text.
+		void SubmitQueryAsync(FString QueryText)
+		{
+			TSharedPtr<IEmbeddingProvider> Provider = ISemanticSearchModule::Get().GetEmbeddingProvider();
+			if (!Provider)
+			{
+				UE_LOG(LogSemanticSearchCB, Warning, TEXT("No embedding provider registered."));
+				return;
+			}
+
+			++ActiveEpoch;
+			bHasInFlightRequest = true;
+			const uint32 Epoch = ActiveEpoch;
+
+			TWeakPtr<FSemanticSearchAssetViewOverride> WeakSelf = AsShared();
+			Provider->GenerateEmbeddingAsync(QueryText, [WeakSelf, Epoch](FEmbeddingResponse&& Response) mutable
+			{
+				// Provider callback may arrive on any thread; marshal to GT before
+				// touching any shared state.
+				AsyncTask(ENamedThreads::GameThread, [WeakSelf, Epoch, Response = MoveTemp(Response)]() mutable
+				{
+					TSharedPtr<FSemanticSearchAssetViewOverride> Self = WeakSelf.Pin();
+					if (!Self || Self->ActiveEpoch != Epoch)
+					{
+						return;
+					}
+					if (Response.ErrorMessage.IsEmpty())
+					{
+						Self->OnQueryEmbeddingReady(MoveTemp(Response.Embedding), Epoch);
+					}
+					else
+					{
+						UE_LOG(LogSemanticSearchCB, Warning, TEXT("Query embedding failed: %s"), *Response.ErrorMessage);
+						Self->bHasInFlightRequest = false;
+					}
+				});
+			}, nullptr);
+		}
+
+		void OnQueryEmbeddingReady(TArray<float> QueryEmbedding, uint32 Epoch)
+		{
+			if (CurrentItems.IsEmpty())
+			{
+				bHasInFlightRequest = false;
+				return;
+			}
+
+			// Shared embedding avoids copying the query vector once per item callback.
+			TSharedRef<const TArray<float>> SharedQuery = MakeShared<const TArray<float>>(MoveTemp(QueryEmbedding));
+
+			using FScoreEntry = TPair<TSharedPtr<FAssetViewItem>, float>;
+			TSharedRef<TArray<FScoreEntry>> Scores = MakeShared<TArray<FScoreEntry>>();
+			TSharedRef<int32> Countdown = MakeShared<int32>(CurrentItems.Num());
+
+			TWeakPtr<FSemanticSearchAssetViewOverride> WeakSelf = AsShared();
+
+			for (const TSharedPtr<FAssetViewItem>& Item : CurrentItems)
+			{
+				FAssetData AssetData;
+				if (!Item->GetItem().Legacy_TryGetAssetData(AssetData))
+				{
+					if (--(*Countdown) == 0)
+					{
+						PublishScores(MoveTemp(*Scores));
+					}
+					continue;
+				}
+
+				FAssetProcessorManager::Get().GetEmbeddingData(AssetData,
+					[WeakSelf, Epoch, Countdown, Scores, SharedQuery, Item]
+					(FAssetEmbeddingResult&& Result, FString&& Error) mutable
+				{
+					AsyncTask(ENamedThreads::GameThread,
+						[WeakSelf, Epoch, Countdown, Scores, SharedQuery, Item,
+						 Result = MoveTemp(Result), Error = MoveTemp(Error)]() mutable
+					{
+						TSharedPtr<FSemanticSearchAssetViewOverride> Self = WeakSelf.Pin();
+						if (!Self || Self->ActiveEpoch != Epoch)
+						{
+							return;
+						}
+
+						if (Error.IsEmpty() && !Result.Embedding.IsEmpty())
+						{
+							const float Similarity = CosineSimilarity(*SharedQuery, Result.Embedding);
+							Scores->Add({Item, Similarity});
+						}
+
+						if (--(*Countdown) == 0)
+						{
+							Self->PublishScores(MoveTemp(*Scores));
+						}
+					});
+				});
+			}
+		}
+
+		void PublishScores(TArray<TPair<TSharedPtr<FAssetViewItem>, float>>&& Scores)
+		{
+			Scores.Sort([](const TPair<TSharedPtr<FAssetViewItem>, float>& A,
+			               const TPair<TSharedPtr<FAssetViewItem>, float>& B)
+			{
+				return A.Value > B.Value;
+			});
+
+			LastMatchedItems.Reset(Scores.Num());
+			for (auto& [Item, Score] : Scores)
+			{
+				LastMatchedItems.Add(MoveTemp(Item));
+			}
+
+			bHasInFlightRequest = false;
+			bHasActiveResults = true;
+			PublishMatchingItems();
+		}
+
+		void PublishMatchingItems()
+		{
+			PublishResults(MakeArrayView(LastMatchedItems));
+		}
+
+		// When no query is active, surface every item that has asset data.
+		// Folder items and other non-asset entries are excluded.
+		void PublishPassThrough()
+		{
+			TArray<TSharedPtr<FAssetViewItem>> PassedItems;
+			PassedItems.Reserve(CurrentItems.Num());
+			for (const TSharedPtr<FAssetViewItem>& Item : CurrentItems)
+			{
+				FAssetData AssetData;
+				if (Item->GetItem().Legacy_TryGetAssetData(AssetData))
+				{
+					PassedItems.Add(Item);
+				}
+			}
+			PublishResults(MakeArrayView(PassedItems));
+		}
+
+		FText ActiveQuery;
+		TArray<TSharedPtr<FAssetViewItem>> CurrentItems;
+		TArray<TSharedPtr<FAssetViewItem>> LastMatchedItems;
+		uint32 ActiveEpoch = 0;
+		bool bHasInFlightRequest = false;
+		bool bHasActiveResults = false;
+	};
+
+	// ---------------------------------------------------------------------------
+
+	class FSemanticSearchCBOverride : public IAssetSearchOverride
+	{
+	public:
+		virtual TSharedRef<ICBSearchOverrideInstance> MakePerViewOverride() const override
+		{
+			return MakeShared<FSemanticSearchAssetViewOverride>();
+		}
+	};
+
+	// ---------------------------------------------------------------------------
+
+	TSharedRef<SWidget> FSemanticSearchAssetViewOverride::GetOverrideModeToggleWidget()
+	{
+		return SNew(SButton)
+			.ButtonStyle(FAppStyle::Get(), "HoverOnlyButton")
+			.ContentPadding(FMargin(4.f, 0.f))
+			.ToolTipText(NSLOCTEXT("SemanticSearch", "ToggleButtonTooltip", "Toggle semantic search"))
+			.OnClicked_Lambda([this]()
+			{
+				ToggleOverride();
+				return FReply::Handled();
+			})
+			[
+				SNew(STextBlock)
+				.TextStyle(FAppStyle::Get(), "SmallText")
+				.Text(NSLOCTEXT("SemanticSearch", "ToggleButtonLabel", "AI"))
+				.ColorAndOpacity_Lambda([this]()
+				{
+					return IsOverrideActive()
+						? FSlateColor(FLinearColor(0.f, 0.8f, 1.f))
+						: FSlateColor::UseForeground();
+				})
+			];
+	}
+
+} // namespace Private
+
+// ---------------------------------------------------------------------------
+
+void RegisterContentBrowserExtension()
+{
+	FContentBrowserModule& ContentBrowserModule =
+		FModuleManager::LoadModuleChecked<FContentBrowserModule>("ContentBrowser");
+
+	PRAGMA_DISABLE_EXPERIMENTAL_WARNINGS
+	const bool bRegistered = ContentBrowserModule.RegisterAssetSearchOverride(MakeUnique<Private::FSemanticSearchCBOverride>());
+	PRAGMA_ENABLE_EXPERIMENTAL_WARNINGS
+
+	ensureMsgf(bRegistered, TEXT("SemanticSearch: failed to register asset search override — another override is already active."));
+}
+
+void UnregisterContentBrowserExtension()
+{
+	if (FContentBrowserModule* ContentBrowserModule = FModuleManager::GetModulePtr<FContentBrowserModule>("ContentBrowser"))
+	{
+		PRAGMA_DISABLE_EXPERIMENTAL_WARNINGS
+		ContentBrowserModule->UnregisterAssetSearchOverride();
+		PRAGMA_ENABLE_EXPERIMENTAL_WARNINGS
+	}
+}
+
+} // namespace UE::SemanticSearch::ContentBrowser
diff --git a/Engine/Plugins/Experimental/SemanticSearch/Source/SemanticSearchEditorIntegrations/Private/SemanticSearchEditorIntegrationsModule.cpp b/Engine/Plugins/Experimental/SemanticSearch/Source/SemanticSearchEditorIntegrations/Private/SemanticSearchEditorIntegrationsModule.cpp
new file mode 100644
--- /dev/null
+++ b/Engine/Plugins/Experimental/SemanticSearch/Source/SemanticSearchEditorIntegrations/Private/SemanticSearchEditorIntegrationsModule.cpp
@@ -0,0 +1,29 @@
+// Copyright Epic Games, Inc. All Rights Reserved.
+
+#include "ISemanticSearchEditorIntegrationsModule.h"
+
+#include "ContentBrowserModule.h"
+#include "ContentBrowser/SemanticSearchCB.h"
+#include "Modules/ModuleManager.h"
+
+namespace UE::SemanticSearch
+{
+	class FSemanticSearchEditorIntegrationsModule : public ISemanticSearchEditorIntegrationsModule
+	{
+	public:
+		virtual void StartupModule() override;
+		virtual void ShutdownModule() override;
+	};
+
+	void FSemanticSearchEditorIntegrationsModule::StartupModule()
+	{
+		ContentBrowser::RegisterContentBrowserExtension();
+	}
+
+	void FSemanticSearchEditorIntegrationsModule::ShutdownModule()
+	{
+		ContentBrowser::UnregisterContentBrowserExtension();
+	}
+}
+
+IMPLEMENT_MODULE(UE::SemanticSearch::FSemanticSearchEditorIntegrationsModule, SemanticSearchEditorIntegrations)
diff --git a/Engine/Plugins/Experimental/SemanticSearch/Source/SemanticSearchEditorIntegrations/Public/ContentBrowser/SemanticSearchCB.h b/Engine/Plugins/Experimental/SemanticSearch/Source/SemanticSearchEditorIntegrations/Public/ContentBrowser/SemanticSearchCB.h
new file mode 100644
--- /dev/null
+++ b/Engine/Plugins/Experimental/SemanticSearch/Source/SemanticSearchEditorIntegrations/Public/ContentBrowser/SemanticSearchCB.h
@@ -0,0 +1,11 @@
+// Copyright Epic Games, Inc. All Rights Reserved.
+
+#pragma once
+
+#include "CoreMinimal.h"
+
+namespace UE::SemanticSearch::ContentBrowser
+{
+	SEMANTICSEARCHEDITORINTEGRATIONS_API void RegisterContentBrowserExtension();
+	SEMANTICSEARCHEDITORINTEGRATIONS_API void UnregisterContentBrowserExtension();
+}
\ No newline at end of file
diff --git a/Engine/Plugins/Experimental/SemanticSearch/Source/SemanticSearchEditorIntegrations/Public/ISemanticSearchEditorIntegrationsModule.h b/Engine/Plugins/Experimental/SemanticSearch/Source/SemanticSearchEditorIntegrations/Public/ISemanticSearchEditorIntegrationsModule.h
new file mode 100644
--- /dev/null
+++ b/Engine/Plugins/Experimental/SemanticSearch/Source/SemanticSearchEditorIntegrations/Public/ISemanticSearchEditorIntegrationsModule.h
@@ -0,0 +1,19 @@
+// Copyright Epic Games, Inc. All Rights Reserved.
+
+#pragma once
+
+#include "Modules/ModuleInterface.h"
+#include "Modules/ModuleManager.h"
+
+namespace UE::SemanticSearch
+{
+	class ISemanticSearchEditorIntegrationsModule : public IModuleInterface
+	{
+	public:
+		static ISemanticSearchEditorIntegrationsModule& Get()
+		{
+			static const FName ModuleName = "SemanticSearchEditorIntegrations";
+			return FModuleManager::LoadModuleChecked<ISemanticSearchEditorIntegrationsModule>(ModuleName);
+		}
+	};
+}
diff --git a/Engine/Plugins/Experimental/SemanticSearch/Source/SemanticSearchEditorIntegrations/SemanticSearchEditorIntegrations.Build.cs b/Engine/Plugins/Experimental/SemanticSearch/Source/SemanticSearchEditorIntegrations/SemanticSearchEditorIntegrations.Build.cs
new file mode 100644
--- /dev/null
+++ b/Engine/Plugins/Experimental/SemanticSearch/Source/SemanticSearchEditorIntegrations/SemanticSearchEditorIntegrations.Build.cs
@@ -0,0 +1,33 @@
+// Copyright Epic Games, Inc. All Rights Reserved.
+
+using UnrealBuildTool;
+
+/**
+ * Editor module responsible for integrating the semantic search system into
+ * editor-facing surfaces such as the Content Browser.
+ *
+ * UI and higher-level editor logic lives here; core search logic lives in SemanticSearch.
+ */
+public class SemanticSearchEditorIntegrations : ModuleRules
+{
+	public SemanticSearchEditorIntegrations(ReadOnlyTargetRules Target) : base(Target)
+	{
+		PrivateDependencyModuleNames.AddRange(
+			new string[]
+			{
+				"ContentBrowser",
+				"ContentBrowserData",
+				"CoreUObject",
+				"SemanticSearch",
+				"Slate",
+				"SlateCore",
+				"UnrealEd",
+			});
+
+		PublicDependencyModuleNames.AddRange(
+			new string[]
+			{
+				"Core"
+			});
+	}
+}
diff --git a/Engine/Source/Developer/ToolWidgets/Private/Filters/SFilterSearchBox.cpp b/Engine/Source/Developer/ToolWidgets/Private/Filters/SFilterSearchBox.cpp
--- a/Engine/Source/Developer/ToolWidgets/Private/Filters/SFilterSearchBox.cpp
+++ b/Engine/Source/Developer/ToolWidgets/Private/Filters/SFilterSearchBox.cpp
@@ -60,6 +60,14 @@ public:
 		/** Delegate to get the visibility of the search history chevron */
 		SLATE_EVENT(FGetSearchHistoryVisibility, GetSearchHistoryVisibility)
 
+		/**
+		 * Optional widget to display on the right side of the search box, after the existing buttons.
+		 * EXPERIMENTAL: This API is subject to change.
+		 */
+		PRAGMA_DISABLE_EXPERIMENTAL_WARNINGS
+		SLATE_ARGUMENT(TSharedPtr<SWidget>, ExtraWidget)
+		PRAGMA_ENABLE_EXPERIMENTAL_WARNINGS
+
 	SLATE_END_ARGS()
 	
 	void Construct( const FArguments& InArgs );
@@ -147,7 +155,18 @@ void SFilterSearchBoxImpl::Construct(const FArguments& InArgs)
 				.ColorAndOpacity(FSlateColor::UseForeground())
 			]
 		];
-	
+
+	if (InArgs._ExtraWidget.IsValid())
+	{
+		Box->AddSlot()
+			.HAlign(HAlign_Right)
+			.VAlign(VAlign_Center)
+			.AutoWidth()
+			.Padding(0.f, 0.f, 4.f, 0.f)
+			[
+				InArgs._ExtraWidget.ToSharedRef()
+			];
+	}
 }
 
 void SFilterSearchBoxImpl::SetOnSaveSearchHandler(SFilterSearchBox::FOnSaveSearchClicked InOnSaveSearchHandler)
@@ -194,6 +213,7 @@ void SFilterSearchBox::Construct( const FArguments& InArgs )
 			.OnSearchHistoryClicked(this, &SFilterSearchBox::OnClickedSearchHistory)
 			.GetSearchHistoryVisibility(this, &SFilterSearchBox::GetSearchHistoryVisibility)
 			.OnSaveSearchClicked(InArgs._OnSaveSearchClicked)
+			.ExtraWidget(InArgs._ExtraWidget)
 		]
 		.MenuContent
 		(
diff --git a/Engine/Source/Developer/ToolWidgets/Public/Filters/SFilterSearchBox.h b/Engine/Source/Developer/ToolWidgets/Public/Filters/SFilterSearchBox.h
--- a/Engine/Source/Developer/ToolWidgets/Public/Filters/SFilterSearchBox.h
+++ b/Engine/Source/Developer/ToolWidgets/Public/Filters/SFilterSearchBox.h
@@ -69,6 +69,14 @@ public:
 	/** Handler for when the + Button next to a search is clicked */
 	SLATE_EVENT(FOnSaveSearchClicked, OnSaveSearchClicked)
 
+	/**
+	 * Optional widget to display on the right side of the search box, after the existing buttons.
+	 * EXPERIMENTAL: This API is subject to change.
+	 */
+	PRAGMA_DISABLE_EXPERIMENTAL_WARNINGS
+	SLATE_ARGUMENT(TSharedPtr<SWidget>, ExtraWidget)
+	PRAGMA_ENABLE_EXPERIMENTAL_WARNINGS
+
 	SLATE_END_ARGS()
 
 	/** Constructs this widget with InArgs */
diff --git a/Engine/Source/Editor/ContentBrowser/Private/ContentBrowserModule.cpp b/Engine/Source/Editor/ContentBrowser/Private/ContentBrowserModule.cpp
--- a/Engine/Source/Editor/ContentBrowser/Private/ContentBrowserModule.cpp
+++ b/Engine/Source/Editor/ContentBrowser/Private/ContentBrowserModule.cpp
@@ -135,3 +135,27 @@ const TArray<FContentBrowserModule::FWizard>& FContentBrowserModule::GetWizards(
 {
 	return Wizards;
 }
+
+bool FContentBrowserModule::RegisterAssetSearchOverride(TUniquePtr<UE::Editor::ContentBrowser::Extension::IAssetSearchOverride>&& InOverride)
+{
+	if (ActiveAssetSearchOverride.IsValid())
+	{
+		return false;
+	}
+
+	ActiveAssetSearchOverride = MoveTemp(InOverride);
+	// Next: Refresh existing content browsers?
+	return true;
+}
+
+void FContentBrowserModule::UnregisterAssetSearchOverride()
+{
+	ActiveAssetSearchOverride.Reset();
+	// Next: Refresh existing content browsers?
+}
+
+UE::Editor::ContentBrowser::Extension::IAssetSearchOverride* FContentBrowserModule::GetActiveAssetSearchOverride() const
+{
+	return ActiveAssetSearchOverride.Get();
+}
+
diff --git a/Engine/Source/Editor/ContentBrowser/Private/Experimental/ContentBrowserSearchOverride.cpp b/Engine/Source/Editor/ContentBrowser/Private/Experimental/ContentBrowserSearchOverride.cpp
new file mode 100644
--- /dev/null
+++ b/Engine/Source/Editor/ContentBrowser/Private/Experimental/ContentBrowserSearchOverride.cpp
@@ -0,0 +1,57 @@
+// Copyright Epic Games, Inc. All Rights Reserved.
+
+#include "Experimental/ContentBrowserSearchOverride.h"
+#include "Experimental/ContentBrowserSearchOverrideInternal.h"
+
+namespace UE::Editor::ContentBrowser::Extension
+{
+
+// ---------------------------------------------------------------------------
+// ICBSearchOverrideInstance
+// ---------------------------------------------------------------------------
+
+ICBSearchOverrideInstance::ICBSearchOverrideInstance() = default;
+ICBSearchOverrideInstance::~ICBSearchOverrideInstance() = default;
+
+bool ICBSearchOverrideInstance::IsOverrideActive() const
+{
+	return bIsOverrideActive;
+}
+
+void ICBSearchOverrideInstance::ToggleOverride()
+{
+	bIsOverrideActive = !bIsOverrideActive;
+	GetInternalApi().OnOverrideActiveChanged.Broadcast();
+}
+
+void ICBSearchOverrideInstance::SetUserSearching(bool bIsSearching)
+{
+	GetInternalApi().OnRequestSetUserSearching.ExecuteIfBound(bIsSearching);
+}
+
+void ICBSearchOverrideInstance::RequestQuickFrontendListRefresh()
+{
+	GetInternalApi().OnRequestQuickFrontendListRefresh.ExecuteIfBound();
+}
+
+void ICBSearchOverrideInstance::PublishResults(TArrayView<TSharedPtr<FAssetViewItem>> Items)
+{
+	GetInternalApi().PublishedResults.Append(Items);
+}
+
+FInternalApiWrapper& ICBSearchOverrideInstance::GetInternalApi()
+{
+	if (!InternalApi)
+	{
+		InternalApi = MakeUnique<FInternalApiWrapper>(*this);
+	}
+	return *InternalApi;
+}
+
+void FInternalApiWrapper::OnFilteringReset()
+{
+	PublishedResults.Reset();
+	Instance.OnFilteringResetImplementation();
+}
+
+} // namespace UE::Editor::ContentBrowser::Extension
diff --git a/Engine/Source/Editor/ContentBrowser/Private/Experimental/ContentBrowserSearchOverrideInternal.h b/Engine/Source/Editor/ContentBrowser/Private/Experimental/ContentBrowserSearchOverrideInternal.h
new file mode 100644
--- /dev/null
+++ b/Engine/Source/Editor/ContentBrowser/Private/Experimental/ContentBrowserSearchOverrideInternal.h
@@ -0,0 +1,38 @@
+// Copyright Epic Games, Inc. All Rights Reserved.
+
+#pragma once
+
+#include "Experimental/ContentBrowserSearchOverride.h"
+
+namespace UE::Editor::ContentBrowser::Extension
+{
+	/**
+	 * Provides write access to the private state of ICBSearchOverrideInstance.
+	 * Only ContentBrowser internals should hold a reference to this wrapper.
+	 */
+	class FInternalApiWrapper
+	{
+	public:
+		explicit FInternalApiWrapper(ICBSearchOverrideInstance& InInstance)
+			: Instance(InInstance)
+		{}
+
+		void OnFilteringReset();
+
+		/** Fires when the override is toggled on or off. */
+		FSimpleMulticastDelegate OnOverrideActiveChanged;
+
+		/** Executed when the override calls SetUserSearching(). */
+		TDelegate<void(bool)> OnRequestSetUserSearching;
+
+		/** Executed when the override calls RequestQuickFrontendListRefresh(). */
+		FSimpleDelegate OnRequestQuickFrontendListRefresh;
+
+		/** Results set by PublishResults(); consumed by the asset view on the next tick. */
+		TArray<TSharedPtr<FAssetViewItem>> PublishedResults;
+
+	private:
+		ICBSearchOverrideInstance& Instance;
+	};
+
+} // namespace UE::Editor::ContentBrowser::Extension
diff --git a/Engine/Source/Editor/ContentBrowser/Private/SAssetView.cpp b/Engine/Source/Editor/ContentBrowser/Private/SAssetView.cpp
--- a/Engine/Source/Editor/ContentBrowser/Private/SAssetView.cpp
+++ b/Engine/Source/Editor/ContentBrowser/Private/SAssetView.cpp
@@ -38,6 +38,7 @@
 #include "EditorWidgetsModule.h"
 #include "Engine/Blueprint.h"
 #include "Engine/Level.h"
+#include "Experimental/ContentBrowserSearchOverrideInternal.h"
 #include "Factories/Factory.h"
 #include "Framework/Application/SlateApplication.h"
 #include "Framework/Commands/GenericCommands.h"
@@ -1032,7 +1033,7 @@ void FAssetViewItemCollection::UpdateItemFiltering(
 		}
 	}
 
-	// Publish items which originally failed filtering and then were updated to pass it 
+	// Publish items which originally failed filtering and then were updated to pass it
 	for (int32 Index : ItemsPendingPriorityPublish)
 	{
 		if (ensureMsgf(!FilterState[Index].Published, TEXT("Priority-publish item %d was already published. PublishProgress: %d CanPublish: %d"), Index, PublishProgress, CanPublish))
@@ -2656,6 +2657,12 @@ void SAssetView::CalculateThumbnailHintColorAndOpacity()
 
 bool SAssetView::HasItemsPendingFilter() const
 {
+	PRAGMA_DISABLE_EXPERIMENTAL_WARNINGS
+	if (AssetSearchOverride.IsValid() && AssetSearchOverride->IsOverrideActive() && (AssetSearchOverride->HasPendingResults() || !AssetSearchOverride->GetInternalApi().PublishedResults.IsEmpty()))
+	{
+		return true;
+	}
+	PRAGMA_ENABLE_EXPERIMENTAL_WARNINGS
 	return Items->HasItemsPendingFilter();
 }
 
@@ -2680,6 +2687,18 @@ void SAssetView::ProcessItemsPendingFilter(const double TickStartTime)
 
 	Items->UpdateItemFiltering(FrontendFilterHelper, bFlushAllPendingItems ? MAX_dbl : TickStartTime + MaxSecondsPerFrame, FilterResults);
 
+	PRAGMA_DISABLE_EXPERIMENTAL_WARNINGS
+		if (AssetSearchOverride.IsValid() && AssetSearchOverride->IsOverrideActive())
+		{
+			AssetSearchOverride->OnItemsAvailable(MakeArrayView(FilterResults));
+			FilterResults.Reset();
+
+			FilterResults.Append(MoveTemp(AssetSearchOverride->GetInternalApi().PublishedResults));
+			AssetSearchOverride->GetInternalApi().PublishedResults.Reset();
+		}
+	PRAGMA_ENABLE_EXPERIMENTAL_WARNINGS
+
+
 	// exclude duplicates
 	for (TSharedPtr<FAssetViewItem>& Result : FilterResults)
 	{
@@ -3781,7 +3800,15 @@ void SAssetView::RefreshFilteredItems()
 	const double RefreshFilteredItemsStartTime = FPlatformTime::Seconds();
 
 	OnInterruptFiltering();
-	
+
+	PRAGMA_DISABLE_EXPERIMENTAL_WARNINGS
+	if (AssetSearchOverride.IsValid() && AssetSearchOverride->IsOverrideActive())
+	{
+		UE::Editor::ContentBrowser::Extension::FInternalApiWrapper& SearchOverrideApi = AssetSearchOverride->GetInternalApi();
+		SearchOverrideApi.OnFilteringReset();
+	}
+	PRAGMA_ENABLE_EXPERIMENTAL_WARNINGS
+
 	FilteredAssetItems.Reset();
 	AlreadyAddedOriginalAssets.Reset();
 	FilteredAssetItemTypeCounts.Reset();
@@ -4396,6 +4423,14 @@ void SAssetView::SortList(bool bSyncToSelection)
 {
 	TRACE_CPUPROFILER_EVENT_SCOPE(SAssetView::SortList);
 
+	PRAGMA_DISABLE_EXPERIMENTAL_WARNINGS
+	if (AssetSearchOverride.IsValid() && AssetSearchOverride->IsOverrideActive())
+	{
+		bPendingSortFilteredItems = false;
+		return;
+	}
+	PRAGMA_ENABLE_EXPERIMENTAL_WARNINGS
+
 	if ( !IsRenamingAsset() )
 	{
 		SortManager->SortList(FilteredAssetItems, MajorityAssetType, CustomColumns);
@@ -6885,6 +6920,23 @@ void SAssetView::SetShouldFilterItem(FOnShouldFilterItem InCallback)
 	RequestQuickFrontendListRefresh();
 }
 
+void SAssetView::SetAssetSearchOverride(TSharedPtr<UE::Editor::ContentBrowser::Extension::ICBSearchOverrideInstance> InOverride)
+{
+	AssetSearchOverride = MoveTemp(InOverride);
+	if (AssetSearchOverride.IsValid())
+	{
+		AssetSearchOverride->GetInternalApi().OnOverrideActiveChanged.AddSP(this, &SAssetView::OnSearchOverrideToggled);
+		AssetSearchOverride->GetInternalApi().OnRequestSetUserSearching.BindSP(this, &SAssetView::SetUserSearching);
+		AssetSearchOverride->GetInternalApi().OnRequestQuickFrontendListRefresh.BindSP(this, &SAssetView::RequestQuickFrontendListRefresh);
+	}
+	RequestQuickFrontendListRefresh();
+}
+
+void SAssetView::OnSearchOverrideToggled()
+{
+	RequestQuickFrontendListRefresh();
+}
+
 TWeakPtr<FAssetViewSortManager> SAssetView::GetSortManager() const
 {
 	return SortManager;
diff --git a/Engine/Source/Editor/ContentBrowser/Private/SContentBrowser.cpp b/Engine/Source/Editor/ContentBrowser/Private/SContentBrowser.cpp
--- a/Engine/Source/Editor/ContentBrowser/Private/SContentBrowser.cpp
+++ b/Engine/Source/Editor/ContentBrowser/Private/SContentBrowser.cpp
@@ -136,6 +136,7 @@
 #include "ContentSources/Widgets/SContentSourcesView.h"
 #include "ContentSources/Widgets/SLegacyContentSource.h"
 #include "Experimental/ContentBrowserExtensionUtils.h"
+#include "Experimental/ContentBrowserSearchOverrideInternal.h"
 #include "Misc/EnumerateRange.h"
 #include "Toolkits/GlobalEditorCommonCommands.h"
 #include "Trace/Detail/Channel.h"
@@ -1025,6 +1026,26 @@ TSharedRef<SWidget> SContentBrowser::CreateAssetView(const FContentBrowserConfig
 
 	LegacyContentSourceWidgets->AssetViewPtr->SetFilterBar(LegacyContentSourceWidgets->FilterListPtr);
 
+	// If a search override is registered, create a per-view instance and retrieve its toggle widget
+	PRAGMA_DISABLE_EXPERIMENTAL_WARNINGS
+	if (UE::Editor::ContentBrowser::Extension::IAssetSearchOverride* ActiveOverride = FModuleManager::LoadModuleChecked<FContentBrowserModule>("ContentBrowser").GetActiveAssetSearchOverride())
+	{
+		LegacyContentSourceWidgets->AssetSearchOverrideInstance = ActiveOverride->MakePerViewOverride();
+
+		// When the override is toggled, clear the search text and reset the filter pipeline.
+		LegacyContentSourceWidgets->AssetSearchOverrideInstance->GetInternalApi().OnOverrideActiveChanged.AddSP(this, &SContentBrowser::ClearSearchText);
+
+		// Pass the instance to the asset view so it can take over the last step of the filtering and the sort.
+		LegacyContentSourceWidgets->AssetViewPtr->SetAssetSearchOverride(LegacyContentSourceWidgets->AssetSearchOverrideInstance);
+	}
+
+	TSharedPtr<SWidget> SearchOverrideToggleWidget;
+	if (LegacyContentSourceWidgets->AssetSearchOverrideInstance.IsValid())
+	{
+		SearchOverrideToggleWidget = LegacyContentSourceWidgets->AssetSearchOverrideInstance->GetOverrideModeToggleWidget();
+	}
+	PRAGMA_ENABLE_EXPERIMENTAL_WARNINGS
+
 	LegacyContentSourceWidgets->SearchBoxPtr = SNew(SAssetSearchBox)
 				   .SuggestionListPlacement(EMenuPlacement::MenuPlacement_AboveAnchor) // Avoid showing suggestion on filtered content browser results...
 				   .HintText(this, &SContentBrowser::GetSearchAssetsHintText)
@@ -1036,6 +1057,7 @@ TSharedRef<SWidget> SContentBrowser::CreateAssetView(const FContentBrowserConfig
 				   .OnAssetSearchBoxSuggestionFilter(this, &SContentBrowser::OnAssetSearchSuggestionFilter)
 				   .OnAssetSearchBoxSuggestionChosen(this, &SContentBrowser::OnAssetSearchSuggestionChosen)
 				   .DelayChangeNotificationsWhileTyping(true)
+				   .ExtraWidget(SearchOverrideToggleWidget)
 				   .Visibility((Config != nullptr ? Config->bCanShowAssetSearch : true) ? EVisibility::Visible : EVisibility::Collapsed)
 				   .AddMetaData<FTagMetaData>(FTagMetaData(TEXT("ContentBrowserSearchAssets")));
 
@@ -2402,6 +2424,20 @@ void SContentBrowser::SetSearchText(const FText& InSearchText)
 	}
 }
 
+void SContentBrowser::ClearSearchText()
+{
+	SetSearchBoxText(FText::GetEmpty());
+	if (LegacyContentSourceWidgets->SearchBoxPtr)
+	{
+		LegacyContentSourceWidgets->SearchBoxPtr->SetText(FText::GetEmpty());
+	}
+}
+
+void SContentBrowser::SetAssetViewUserSearching(bool bIsSearching)
+{
+	LegacyContentSourceWidgets->AssetViewPtr->SetUserSearching(bIsSearching);
+}
+
 void SContentBrowser::SetSearchBoxText(const FText& InSearchText)
 {
 	// Has anything changed? (need to test case as the operators are case-sensitive)
@@ -2422,16 +2458,44 @@ void SContentBrowser::SetSearchBoxText(const FText& InSearchText)
 
 void SContentBrowser::OnSearchBoxChanged(const FText& InSearchText)
 {
-	SetSearchBoxText(InSearchText);
+	bool bOverrideHandled = false;
+
+	PRAGMA_DISABLE_EXPERIMENTAL_WARNINGS
+	if (LegacyContentSourceWidgets->AssetSearchOverrideInstance.IsValid() &&
+		LegacyContentSourceWidgets->AssetSearchOverrideInstance->IsOverrideActive())
+	{
+		LegacyContentSourceWidgets->AssetSearchOverrideInstance->OnSearchTextChanged(InSearchText);
+		bOverrideHandled = true;
+	}
+	PRAGMA_ENABLE_EXPERIMENTAL_WARNINGS
+
+	if (!bOverrideHandled)
+	{
+		SetSearchBoxText(InSearchText);
+	}
 
 	// Broadcast 'search box changed' delegate
 	FContentBrowserModule& ContentBrowserModule = FModuleManager::GetModuleChecked<FContentBrowserModule>( TEXT("ContentBrowser") );
-	ContentBrowserModule.GetOnSearchBoxChanged().Broadcast(InSearchText, bIsPrimaryBrowser);	
+	ContentBrowserModule.GetOnSearchBoxChanged().Broadcast(InSearchText, bIsPrimaryBrowser);
 }
 
 void SContentBrowser::OnSearchBoxCommitted(const FText& InSearchText, ETextCommit::Type CommitInfo)
 {
-	SetSearchBoxText(InSearchText);
+	bool bOverrideHandled = false;
+
+	PRAGMA_DISABLE_EXPERIMENTAL_WARNINGS
+	if (LegacyContentSourceWidgets->AssetSearchOverrideInstance.IsValid() &&
+		LegacyContentSourceWidgets->AssetSearchOverrideInstance->IsOverrideActive())
+	{
+		LegacyContentSourceWidgets->AssetSearchOverrideInstance->OnSearchTextCommitted(InSearchText, CommitInfo);
+		bOverrideHandled = true;
+	}
+	PRAGMA_ENABLE_EXPERIMENTAL_WARNINGS
+
+	if (!bOverrideHandled)
+	{
+		SetSearchBoxText(InSearchText);
+	}
 }
 
 FReply SContentBrowser::OnSearchKeyDown(const FGeometry& Geometry, const FKeyEvent& InKeyEvent)
diff --git a/Engine/Source/Editor/ContentBrowser/Private/SContentBrowser.h b/Engine/Source/Editor/ContentBrowser/Private/SContentBrowser.h
--- a/Engine/Source/Editor/ContentBrowser/Private/SContentBrowser.h
+++ b/Engine/Source/Editor/ContentBrowser/Private/SContentBrowser.h
@@ -7,6 +7,7 @@
 #include "Containers/MRUArray.h"
 #include "Containers/UnrealString.h"
 #include "ContentBrowserDataMenuContexts.h"
+#include "Experimental/ContentBrowserSearchOverride.h"
 #include "ContentBrowserDataSubsystem.h"
 #include "ContentBrowserDelegates.h"
 #include "Filters/CustomTextFilters.h"
@@ -339,6 +340,12 @@ private:
 	/** Called when the editable text needs to be set or cleared */
 	void SetSearchBoxText(const FText& InSearchText);
 
+	/** Clears the search box text and resets the filter pipeline; used when an asset search override is toggled. */
+	void ClearSearchText();
+
+	/** Forwards the user-searching state to the asset view */
+	void SetAssetViewUserSearching(bool bIsSearching);
+
 	/** Called by the editable text control when the search text is changed by the user */
 	void OnSearchBoxChanged(const FText& InSearchText);
 
@@ -724,6 +731,9 @@ private:
 		/** The text box used to search for assets */
 		TSharedPtr<SAssetSearchBox> SearchBoxPtr;
 
+		/** Per-view instance of the active search override, if any */
+		TSharedPtr<UE::Editor::ContentBrowser::Extension::ICBSearchOverrideInstance> AssetSearchOverrideInstance;
+
 		/** Handles the search box sizing within the toolbar */
 		TSharedPtr<UE::Editor::ContentBrowser::Private::TWidgetDesiredSizeSwitcher<EAxis::X>> SearchBoxSizeSwitcher;
 
diff --git a/Engine/Source/Editor/ContentBrowser/Public/ContentBrowserModule.h b/Engine/Source/Editor/ContentBrowser/Public/ContentBrowserModule.h
--- a/Engine/Source/Editor/ContentBrowser/Public/ContentBrowserModule.h
+++ b/Engine/Source/Editor/ContentBrowser/Public/ContentBrowserModule.h
@@ -10,6 +10,7 @@
 #include "AssetTypeCategories.h"
 #include "IAssetTools.h"
 #include "Delegates/Delegate.h"
+#include "Experimental/ContentBrowserSearchOverride.h"
 #include "Experimental/ContentBrowserViewExtender.h"
 #include "HAL/PlatformCrt.h"
 #include "MRUFavoritesList.h"
@@ -237,6 +238,22 @@ public:
 	CONTENTBROWSER_API void UnregisterWizard(const FName WizardIdentifier);
 	CONTENTBROWSER_API const TArray<FWizard>& GetWizards() const;
 
+
+	PRAGMA_DISABLE_EXPERIMENTAL_WARNINGS
+	UE_EXPERIMENTAL(5.8, "Asset Search Override is experimental.")
+	CONTENTBROWSER_API bool RegisterAssetSearchOverride(TUniquePtr<UE::Editor::ContentBrowser::Extension::IAssetSearchOverride>&& InOverride);
+	PRAGMA_ENABLE_EXPERIMENTAL_WARNINGS
+
+	PRAGMA_DISABLE_EXPERIMENTAL_WARNINGS
+	UE_EXPERIMENTAL(5.8, "Asset Search Override is experimental.")
+	CONTENTBROWSER_API void UnregisterAssetSearchOverride();
+	PRAGMA_ENABLE_EXPERIMENTAL_WARNINGS
+
+	PRAGMA_DISABLE_EXPERIMENTAL_WARNINGS
+	UE_EXPERIMENTAL(5.8, "Asset Search Override is experimental.")
+	CONTENTBROWSER_API UE::Editor::ContentBrowser::Extension::IAssetSearchOverride* GetActiveAssetSearchOverride() const;
+	PRAGMA_ENABLE_EXPERIMENTAL_WARNINGS
+
 private:
 	TArray<FWizard> Wizards;
 	
@@ -290,4 +307,7 @@ private:
 
 	/** Extension used to provide a custom view to the Content Browser */
 	FCreateViewExtender ContentBrowserViewExtender;
+
+	/** The currently active asset search override, if any. Only one may be registered at a time. */
+	TUniquePtr<UE::Editor::ContentBrowser::Extension::IAssetSearchOverride> ActiveAssetSearchOverride = nullptr;
 };
diff --git a/Engine/Source/Editor/ContentBrowser/Public/Experimental/ContentBrowserSearchOverride.h b/Engine/Source/Editor/ContentBrowser/Public/Experimental/ContentBrowserSearchOverride.h
new file mode 100644
--- /dev/null
+++ b/Engine/Source/Editor/ContentBrowser/Public/Experimental/ContentBrowserSearchOverride.h
@@ -0,0 +1,102 @@
+// Copyright Epic Games, Inc. All Rights Reserved.
+
+#pragma once
+
+#include "Containers/Array.h"
+#include "Containers/ArrayView.h"
+#include "ContentBrowserItem.h"
+#include "Delegates/Delegate.h"
+#include "Internationalization/Text.h"
+#include "Templates/SharedPointer.h"
+#include "Templates/UniquePtr.h"
+#include "Types/SlateEnums.h"
+
+class FAssetViewItem;
+
+namespace UE::Editor::ContentBrowser::Extension
+{
+	/**
+	 * Interface to contains the state per view of the search override
+	 */
+	class ICBSearchOverrideInstance
+	{
+	public:
+		CONTENTBROWSER_API ICBSearchOverrideInstance();
+		CONTENTBROWSER_API virtual ~ICBSearchOverrideInstance();
+
+		CONTENTBROWSER_API bool IsOverrideActive() const;
+		CONTENTBROWSER_API void ToggleOverride();
+
+		/**
+		 * Call from the override to notify the Content Browser to update the asset view's user-searching state.
+		 * Pass true when a search is in progress, false when results are cleared.
+		 */
+		CONTENTBROWSER_API void SetUserSearching(bool bIsSearching);
+
+		/**
+		 * Call from the override to trigger a quick frontend list refresh on the asset view.
+		 * Use this when the search query changes so the filter pipeline re-runs immediately.
+		 */
+		CONTENTBROWSER_API  void RequestQuickFrontendListRefresh();
+
+		virtual TSharedRef<class SWidget> GetOverrideModeToggleWidget() = 0;
+
+		/** Called when the search box text changes, only when IsOverrideActive() returns true. */
+		virtual void OnSearchTextChanged(const FText& NewText) = 0;
+
+		/** Called when the search box text is committed, only when IsOverrideActive() returns true. */
+		virtual void OnSearchTextCommitted(const FText& NewText, ETextCommit::Type CommitType) = 0;
+
+
+		/**
+		 * Called each filter pass with all items that passed the frontend filters.
+		 * The override is responsible for deciding which of these to surface via PublishResults().
+		 */
+		virtual void OnItemsAvailable(TArrayView<const TSharedPtr<FAssetViewItem>> Items) = 0;
+
+
+		virtual void SortItemList(TArrayView<TSharedPtr<FAssetViewItem>> Items) = 0;
+
+		/**
+		 * Returns true while the override has async work in flight whose results
+		 * have not yet been published. The asset view will keep ticking while this is true.
+		 */
+		virtual bool HasPendingResults() const { return false; }
+
+		/**
+		 * Call when results are ready. The asset view will inject these items into the
+		 * visible set on the next tick.
+		 */
+		CONTENTBROWSER_API void PublishResults(TArrayView<TSharedPtr<FAssetViewItem>> Items);
+
+		/** Private API for ContentBrowser internals. */
+		class FInternalApiWrapper& GetInternalApi();
+
+		friend class FInternalApiWrapper;
+
+	protected:
+
+		/**
+		 * Called at the start of each filter pass when the override is active.
+		 * Override to reset any pending state before new items start arriving.
+		 */
+		virtual void OnFilteringResetImplementation() = 0;
+
+	private:
+		bool bIsOverrideActive = false;
+		TUniquePtr<FInternalApiWrapper> InternalApi;
+	};
+
+
+	/**
+	 * Interface for an asset search override that can hook into the Content Browser's
+	 * filter and sort pipeline. Only one override may be active at any given time.
+	 */
+	class IAssetSearchOverride
+	{
+	public:
+		virtual ~IAssetSearchOverride() = default;
+
+		virtual TSharedRef<ICBSearchOverrideInstance> MakePerViewOverride() const = 0;
+	};
+}
diff --git a/Engine/Source/Editor/ContentBrowser/Public/SAssetView.h b/Engine/Source/Editor/ContentBrowser/Public/SAssetView.h
--- a/Engine/Source/Editor/ContentBrowser/Public/SAssetView.h
+++ b/Engine/Source/Editor/ContentBrowser/Public/SAssetView.h
@@ -15,6 +15,7 @@
 #include "ContentBrowserDataFilter.h"
 #include "ContentBrowserDataLegacyBridge.h"
 #include "ContentBrowserDelegates.h"
+#include "Experimental/ContentBrowserSearchOverride.h"
 #include "ContentBrowserItem.h"
 #include "ContentBrowserTelemetry.h"
 #include "AssetViewContentSources.h"
@@ -532,6 +533,15 @@ public:
 	/** Change the delegate bound via the widget argument OnShouldFilterAsset after construction. */
 	UE_API void SetShouldFilterItem(FOnShouldFilterItem InCallback);
 
+	/**
+	 * Pass the active per-view search override instance to the asset view.
+	 * The asset view will intercept items at the publish step and wire publish-results back.
+	 * EXPERIMENTAL: This API is subject to change.
+	 */
+	PRAGMA_DISABLE_EXPERIMENTAL_WARNINGS
+	UE_API void SetAssetSearchOverride(TSharedPtr<UE::Editor::ContentBrowser::Extension::ICBSearchOverrideInstance> InOverride);
+	PRAGMA_ENABLE_EXPERIMENTAL_WARNINGS
+
 	/** True if we have items pending that ProcessItemsPendingFilter still needs to process */
 	UE_API bool HasItemsPendingFilter() const;
 
@@ -1126,6 +1136,14 @@ private:
 	/** Called to check if an item should be filtered out by external code */
 	FOnShouldFilterItem OnShouldFilterItem;
 
+	/** Active per-view search override instance; when set, intercepts items at the publish step. */
+	PRAGMA_DISABLE_EXPERIMENTAL_WARNINGS
+	TSharedPtr<UE::Editor::ContentBrowser::Extension::ICBSearchOverrideInstance> AssetSearchOverride;
+	PRAGMA_ENABLE_EXPERIMENTAL_WARNINGS
+
+	/** Called when the override is toggled; triggers a fresh frontend filter pass. */
+	void OnSearchOverrideToggled();
+
 	/** Called when the asset view is asked to start to create a temporary item */
 	FOnAssetViewNewItemRequested OnNewItemRequested;
 
diff --git a/Engine/Source/Editor/EditorWidgets/Private/SAssetSearchBox.cpp b/Engine/Source/Editor/EditorWidgets/Private/SAssetSearchBox.cpp
--- a/Engine/Source/Editor/EditorWidgets/Private/SAssetSearchBox.cpp
+++ b/Engine/Source/Editor/EditorWidgets/Private/SAssetSearchBox.cpp
@@ -97,6 +97,7 @@ void SAssetSearchBox::Construct( const FArguments& InArgs )
 				.OnKeyDownHandler(this, &SAssetSearchBox::HandleKeyDown)
 				.ShowSearchHistory(InArgs._ShowSearchHistory)
 				.OnSaveSearchClicked(InArgs._OnSaveSearchClicked)
+				.ExtraWidget(InArgs._ExtraWidget)
 			]
 			.MenuContent
 			(
diff --git a/Engine/Source/Editor/EditorWidgets/Public/SAssetSearchBox.h b/Engine/Source/Editor/EditorWidgets/Public/SAssetSearchBox.h
--- a/Engine/Source/Editor/EditorWidgets/Public/SAssetSearchBox.h
+++ b/Engine/Source/Editor/EditorWidgets/Public/SAssetSearchBox.h
@@ -127,6 +127,14 @@ public:
 		/** Height of the suggestion list dropdown */
 		SLATE_ARGUMENT(float, SuggestionListHeight)
 
+		/**
+		 * Optional widget to display on the right side of the search box, after the existing buttons.
+		 * EXPERIMENTAL: This API is subject to change.
+		 */
+		PRAGMA_DISABLE_EXPERIMENTAL_WARNINGS
+		SLATE_ARGUMENT(TSharedPtr<SWidget>, ExtraWidget)
+		PRAGMA_ENABLE_EXPERIMENTAL_WARNINGS
+
 	SLATE_END_ARGS()
 
 	/** Constructs this widget with InArgs */
