Extend the SemanticSearch material instance processor to handle derived material instance classes by enabling subtype support

#Synced-CL 52352994

diff --git a/Engine/Plugins/Experimental/SemanticSearch/Source/SemanticSearch/Private/AssetProcessors/MaterialInstanceAssetProcessor.cpp b/Engine/Plugins/Experimental/SemanticSearch/Source/SemanticSearch/Private/AssetProcessors/MaterialInstanceAssetProcessor.cpp
--- a/Engine/Plugins/Experimental/SemanticSearch/Source/SemanticSearch/Private/AssetProcessors/MaterialInstanceAssetProcessor.cpp
+++ b/Engine/Plugins/Experimental/SemanticSearch/Source/SemanticSearch/Private/AssetProcessors/MaterialInstanceAssetProcessor.cpp
@@ -19,7 +19,7 @@ UClass& FMaterialInstanceProcessor::GetSupportedClass() const
 
 bool FMaterialInstanceProcessor::SupportDerivedClasses() const
 {
-	return false;
+	return true;
 }
 
 TSharedPtr<FJsonObject> FMaterialInstanceProcessor::GetMetadata(const TSharedRef<const FAssetData>& InAsset) const
