Add rate limiting and retry logic in RemoteEmbeddingProvider.
API stays the same

#Synced-CL 52357265

diff --git a/Engine/Plugins/Experimental/SemanticSearch/Source/SemanticSearch/Private/RemoteEmbeddingProvider.cpp b/Engine/Plugins/Experimental/SemanticSearch/Source/SemanticSearch/Private/RemoteEmbeddingProvider.cpp
--- a/Engine/Plugins/Experimental/SemanticSearch/Source/SemanticSearch/Private/RemoteEmbeddingProvider.cpp
+++ b/Engine/Plugins/Experimental/SemanticSearch/Source/SemanticSearch/Private/RemoteEmbeddingProvider.cpp
@@ -28,6 +28,11 @@ public:
 		RequestCompleteDelegate = MoveTemp(HttpRequest->OnProcessRequestComplete());
 		HttpRequest->OnProcessRequestComplete().BindRaw(this, &FSemanticHttpRequest::OnRequestComplete);
 		RequestOwner->Begin(this);
+	}
+
+	/** Start the HTTP request. Separated from construction so Begin() can run in the caller's DDC context. */
+	void Start()
+	{
 		HttpRequest->ProcessRequest();
 	}
 
@@ -98,6 +103,37 @@ private:
 };
 
 
+void FThrottledRequestQueue::Enqueue(TFunction<void()>&& DispatchFn)
+{
+	FScopeLock ScopeLock(&Lock);
+	if (InFlightCount < MaxConcurrent)
+	{
+		++InFlightCount;
+		DispatchFn();
+	}
+	else
+	{
+		PendingRequests.Enqueue(MoveTemp(DispatchFn));
+	}
+}
+
+void FThrottledRequestQueue::OnRequestComplete()
+{
+	TFunction<void()> Next;
+	{
+		FScopeLock ScopeLock(&Lock);
+		if (!PendingRequests.Dequeue(Next))
+		{
+			--InFlightCount;
+		}
+		// else: InFlightCount stays the same — one finished, one starting
+	}
+	if (Next)
+	{
+		Next();
+	}
+}
+
 struct FServiceURLs
 {
 	FString EmbedURL;
@@ -136,13 +172,32 @@ static FServiceURLs ReadURLsFromConfig()
 
 FRemoteEmbeddingProvider::FRemoteEmbeddingProvider()
 	: EmbeddingDimension(0)
+	, RetryManager(MakeShared<FHttpRetrySystem::FManager>(
+		FHttpRetrySystem::RetryLimitCount(MaxRetries),
+		FHttpRetrySystem::RetryTimeoutRelativeSeconds(RetryTimeoutSeconds)))
 {
+	CaptionQueue.MaxConcurrent = MaxConcurrentCaptionRequests;
+	EmbedQueue.MaxConcurrent = MaxConcurrentEmbedRequests;
+
 	FServiceURLs URLs = ReadURLsFromConfig();
 	EmbedURL = MoveTemp(URLs.EmbedURL);
 	CaptionURL = MoveTemp(URLs.CaptionURL);
 	UE_LOG(LogSemanticSearch, Log, TEXT("RemoteEmbeddingProvider created (Embed: %s, Caption: %s)"), *EmbedURL, *CaptionURL);
 }
 
+TSharedRef<IHttpRequest, ESPMode::ThreadSafe> FRemoteEmbeddingProvider::CreateRetryRequest() const
+{
+	FHttpRetrySystem::FRetryResponseCodes RetryResponseCodes({429, 500, 502, 503, 504});
+	FHttpRetrySystem::FRetryVerbs RetryVerbs;
+	RetryVerbs.Add(FName(TEXT("POST")));
+
+	return RetryManager->CreateRequest(
+		FHttpRetrySystem::FRetryLimitCountSetting(),
+		FHttpRetrySystem::FRetryTimeoutRelativeSecondsSetting(),
+		RetryResponseCodes,
+		RetryVerbs);
+}
+
 int32 FRemoteEmbeddingProvider::GetEmbeddingDimension() const
 {
 	return EmbeddingDimension.Load();
@@ -156,7 +211,7 @@ void FRemoteEmbeddingProvider::GenerateCaptionAsync(
 	FString Boundary;
 	TArray<uint8> Body = BuildMultipartBody(Request, Boundary);
 
-	TSharedRef<IHttpRequest, ESPMode::ThreadSafe> Http = FHttpModule::Get().CreateRequest();
+	TSharedRef<IHttpRequest, ESPMode::ThreadSafe> Http = CreateRetryRequest();
 	Http->SetURL(CaptionURL);
 	Http->SetVerb(TEXT("POST"));
 	Http->SetHeader(TEXT("Content-Type"),
@@ -168,6 +223,8 @@ void FRemoteEmbeddingProvider::GenerateCaptionAsync(
 	Http->OnProcessRequestComplete().BindLambda(
 		[this, OnComplete = MakeShared<FOnCaptionComplete>(MoveTemp(InOnComplete))](FHttpRequestPtr /*Req*/, FHttpResponsePtr Resp, bool bConnected)
 		{
+			CaptionQueue.OnRequestComplete();
+
 			FCaptionResponse Response;
 
 			if (!bConnected || !Resp.IsValid())
@@ -225,15 +282,15 @@ void FRemoteEmbeddingProvider::GenerateCaptionAsync(
 			OnComplete->operator()(MoveTemp(Response));
 		});
 
-
 	if (RequestOwner)
 	{
-		//  use ref counting to delete itself
-		new FSemanticHttpRequest(Http, *RequestOwner);
+		// Create DDC request wrapper immediately so Begin() runs in the caller's barrier context
+		auto* DdcRequest = new FSemanticHttpRequest(Http, *RequestOwner);
+		CaptionQueue.Enqueue([DdcRequest]() { DdcRequest->Start(); });
 	}
 	else
 	{
-		Http->ProcessRequest();
+		CaptionQueue.Enqueue([Http]() { Http->ProcessRequest(); });
 	}
 }
 
@@ -249,7 +306,7 @@ void FRemoteEmbeddingProvider::GenerateEmbeddingAsync(
 	TSharedRef<TJsonWriter<>> Writer = TJsonWriterFactory<>::Create(&PayloadStr);
 	FJsonSerializer::Serialize(Payload, Writer);
 
-	TSharedRef<IHttpRequest, ESPMode::ThreadSafe> Http = FHttpModule::Get().CreateRequest();
+	TSharedRef<IHttpRequest, ESPMode::ThreadSafe> Http = CreateRetryRequest();
 	Http->SetURL(EmbedURL);
 	Http->SetVerb(TEXT("POST"));
 	Http->SetHeader(TEXT("Content-Type"), TEXT("application/json"));
@@ -260,6 +317,8 @@ void FRemoteEmbeddingProvider::GenerateEmbeddingAsync(
 	Http->OnProcessRequestComplete().BindLambda(
 		[this, OnComplete = MakeShared<FOnEmbeddingComplete>(MoveTemp(InOnComplete))](FHttpRequestPtr /*Req*/, FHttpResponsePtr Resp, bool bConnected)
 		{
+			EmbedQueue.OnRequestComplete();
+
 			FEmbeddingResponse Response;
 
 			if (!bConnected || !Resp.IsValid())
@@ -337,12 +396,12 @@ void FRemoteEmbeddingProvider::GenerateEmbeddingAsync(
 
 	if (RequestOwner)
 	{
-		//  use ref counting to delete itself
-		new FSemanticHttpRequest(Http, *RequestOwner);
+		auto* DdcRequest = new FSemanticHttpRequest(Http, *RequestOwner);
+		EmbedQueue.Enqueue([DdcRequest]() { DdcRequest->Start(); });
 	}
 	else
 	{
-		Http->ProcessRequest();
+		EmbedQueue.Enqueue([Http]() { Http->ProcessRequest(); });
 	}
 }
 
diff --git a/Engine/Plugins/Experimental/SemanticSearch/Source/SemanticSearch/Private/RemoteEmbeddingProvider.h b/Engine/Plugins/Experimental/SemanticSearch/Source/SemanticSearch/Private/RemoteEmbeddingProvider.h
--- a/Engine/Plugins/Experimental/SemanticSearch/Source/SemanticSearch/Private/RemoteEmbeddingProvider.h
+++ b/Engine/Plugins/Experimental/SemanticSearch/Source/SemanticSearch/Private/RemoteEmbeddingProvider.h
@@ -4,11 +4,32 @@
 
 #include "Interfaces/IEmbeddingProvider.h"
 #include "HttpModule.h"
+#include "HttpRetrySystem.h"
 #include "Interfaces/IHttpRequest.h"
 #include "Interfaces/IHttpResponse.h"
+#include "Containers/Queue.h"
 
 namespace UE::SemanticSearch::Private
 {
+
+/**
+ * Async-friendly concurrency limiter for HTTP requests.
+ * Limits the number of in-flight requests and queues excess for later dispatch.
+ */
+struct FThrottledRequestQueue
+{
+	FCriticalSection Lock;
+	TQueue<TFunction<void()>> PendingRequests;
+	int32 InFlightCount = 0;
+	int32 MaxConcurrent = 8;
+
+	/** Try to dispatch a request. If at capacity, queues it for later. */
+	void Enqueue(TFunction<void()>&& DispatchFn);
+
+	/** Called when a request completes. Dispatches next queued request if any. */
+	void OnRequestComplete();
+};
+
 /**
  *  Remote embedding provider that communicates with MLS backend services.
  *
@@ -28,10 +49,23 @@ private:
 	/** Build multipart/form-data body.*/
 	TArray<uint8> BuildMultipartBody(const FCaptionRequest& Request, FString& OutBoundary) const;
 
+	/** Create an HTTP request with retry configuration. */
+	TSharedRef<IHttpRequest, ESPMode::ThreadSafe> CreateRetryRequest() const;
+
 	static constexpr const TCHAR* TaskPreset = TEXT("UE_ASSET_TAG_AND_CAPTION_V1");
 
 	FString EmbedURL;
 	FString CaptionURL;
 	TAtomic<int32> EmbeddingDimension;
+
+	/** HTTP retry configuration */
+	int32 MaxRetries = 3;
+	double RetryTimeoutSeconds = 120.0;
+	int32 MaxConcurrentCaptionRequests = 1500;
+	int32 MaxConcurrentEmbedRequests = 400;
+
+	TSharedRef<FHttpRetrySystem::FManager> RetryManager;
+	FThrottledRequestQueue CaptionQueue;
+	FThrottledRequestQueue EmbedQueue;
 };
 }
