Add a FHybridSearchIndex singleton which manages the indexing + search instead of the AssetProcessorManager
Add a FBM25Index implementation for keyword search, and implement RRF fusion in FHybridSearchIndex
Add Semantic Search Settings, (enable/disable hybrid search, use flat/pq index etc.)
Change content browser override to use FhybridSearchIndex search

#Synced-CL 52357265

diff --git a/Engine/Plugins/Experimental/SemanticSearch/Source/SemanticSearch/Private/AssetProcessorManager.cpp b/Engine/Plugins/Experimental/SemanticSearch/Source/SemanticSearch/Private/AssetProcessorManager.cpp
--- a/Engine/Plugins/Experimental/SemanticSearch/Source/SemanticSearch/Private/AssetProcessorManager.cpp
+++ b/Engine/Plugins/Experimental/SemanticSearch/Source/SemanticSearch/Private/AssetProcessorManager.cpp
@@ -100,16 +100,6 @@ namespace UE::SemanticSearch
 		});
 	}
 
-	void FAssetProcessorManager::SetVectorIndex(TSharedPtr<IVectorIndex> InIndex)
-	{
-		VectorIndex = MoveTemp(InIndex);
-	}
-
-	TSharedPtr<IVectorIndex> FAssetProcessorManager::GetVectorIndex() const
-	{
-		return VectorIndex;
-	}
-
 	void FAssetProcessorManager::RegisterAssetProcessor(TSharedRef<IAssetProcessor>&& AssetProcessor)
 	{
 		UClass& AssetClass = AssetProcessor->GetSupportedClass();
diff --git a/Engine/Plugins/Experimental/SemanticSearch/Source/SemanticSearch/Private/AssetProcessors/AssetProcessorUtils.cpp b/Engine/Plugins/Experimental/SemanticSearch/Source/SemanticSearch/Private/AssetProcessors/AssetProcessorUtils.cpp
--- a/Engine/Plugins/Experimental/SemanticSearch/Source/SemanticSearch/Private/AssetProcessors/AssetProcessorUtils.cpp
+++ b/Engine/Plugins/Experimental/SemanticSearch/Source/SemanticSearch/Private/AssetProcessors/AssetProcessorUtils.cpp
@@ -69,20 +69,4 @@ void RegisterDefaultAssetProcessors()
 	Manager.RegisterAssetProcessor<FTextureProcessor>();
 }
 
-int64 GetAssetIndexID(const FString& ObjectPath)
-{
-	FBlake3 Hasher;
-	Hasher.Update(ObjectPath.GetCharArray().GetData(), ObjectPath.Len() * sizeof(TCHAR));
-	FBlake3Hash Hash = Hasher.Finalize();
-	int64 Result;
-	FMemory::Memcpy(&Result, Hash.GetBytes(), sizeof(int64));
-	// Ensure non-negative (FAISS uses -1 as sentinel)
-	return Result & 0x7FFFFFFFFFFFFFFF;
-}
-
-int64 GetAssetIndexID(const FAssetData& Asset)
-{
-	return GetAssetIndexID(Asset.GetObjectPathString());
-}
-
 }
diff --git a/Engine/Plugins/Experimental/SemanticSearch/Source/SemanticSearch/Private/AssetProcessors/AssetProcessorUtils.h b/Engine/Plugins/Experimental/SemanticSearch/Source/SemanticSearch/Private/AssetProcessors/AssetProcessorUtils.h
--- a/Engine/Plugins/Experimental/SemanticSearch/Source/SemanticSearch/Private/AssetProcessors/AssetProcessorUtils.h
+++ b/Engine/Plugins/Experimental/SemanticSearch/Source/SemanticSearch/Private/AssetProcessors/AssetProcessorUtils.h
@@ -56,11 +56,4 @@ void PopulateFromEnum(TMap<FString, FString>& OutMap, const UEnum* Enum);
  */
 void RegisterDefaultAssetProcessors();
 
-/**
- * Compute a stable int64 hash of an asset's object path, suitable for use as a vector index ID.
- * Uses Blake3. Always non-negative (FAISS reserves -1 as sentinel).
- */
-int64 GetAssetIndexID(const FAssetData& Asset);
-int64 GetAssetIndexID(const FString& ObjectPath);
-
 }
diff --git a/Engine/Plugins/Experimental/SemanticSearch/Source/SemanticSearch/Private/HybridSearchIndex.cpp b/Engine/Plugins/Experimental/SemanticSearch/Source/SemanticSearch/Private/HybridSearchIndex.cpp
new file mode 100644
--- /dev/null
+++ b/Engine/Plugins/Experimental/SemanticSearch/Source/SemanticSearch/Private/HybridSearchIndex.cpp
@@ -0,0 +1,272 @@
+// Copyright Epic Games, Inc. All Rights Reserved.
+
+#include "HybridSearchIndex.h"
+
+#include "AssetRegistry/AssetData.h"
+#include "Hash/Blake3.h"
+#include "Settings/SemanticSearchSettings.h"
+#include "TextIndex/BM25Index.h"
+#include "VectorIndex/FlatVectorIndex.h"
+#include "VectorIndex/PQVectorIndex.h"
+
+#include "Misc/ScopeRWLock.h"
+
+namespace UE::SemanticSearch
+{
+
+FHybridSearchIndex& FHybridSearchIndex::Get()
+{
+	static FHybridSearchIndex Instance;
+	return Instance;
+}
+
+FHybridSearchIndex::FHybridSearchIndex()
+	: BM25Index(MakeUnique<FBM25Index>())
+{
+}
+
+void FHybridSearchIndex::EnsureInitialized(int32 EmbeddingDimension)
+{
+
+	if (bInitialized)
+	{
+		return;
+	}
+
+	FWriteScopeLock WriteLock(IndexLock);
+	if (bInitialized)
+	{
+		return;
+	}
+
+	const USemanticSearchSettings* Settings = USemanticSearchSettings::Get();
+
+	if (Settings && Settings->IndexType == ESemanticSearchIndexType::PQ)
+	{
+		VectorIndex = MakeShared<FPQVectorIndex>(EmbeddingDimension);
+	}
+	else
+	{
+		VectorIndex = MakeShared<FFlatVectorIndex>(EmbeddingDimension);
+	}
+
+	bInitialized = true;
+}
+
+TSharedPtr<IVectorIndex> FHybridSearchIndex::GetVectorIndex() const
+{
+	return VectorIndex;
+}
+
+bool FHybridSearchIndex::IsHybridEnabled() const
+{
+	const USemanticSearchSettings* Settings = USemanticSearchSettings::Get();
+	return Settings && Settings->bEnableHybridSearch;
+}
+
+FBM25Index& FHybridSearchIndex::GetBM25Index()
+{
+	return *BM25Index;
+}
+
+int64 GetAssetIndexID(FStringView AssetPath)
+{
+	FBlake3 Hasher;
+	Hasher.Update(AssetPath.GetData(), AssetPath.Len() * sizeof(TCHAR));
+	FBlake3Hash Hash = Hasher.Finalize();
+	int64 Result;
+	FMemory::Memcpy(&Result, Hash.GetBytes(), sizeof(int64));
+	// Ensure non-negative (FAISS uses -1 as sentinel)
+	return Result & 0x7FFFFFFFFFFFFFFF;
+}
+
+int64 GetAssetIndexID(const FAssetData& Asset)
+{
+	return GetAssetIndexID(Asset.GetObjectPathString());
+}
+
+void FHybridSearchIndex::Add(
+	const FAssetData& Asset,
+	TConstArrayView<float> Embedding,
+	FStringView Caption,
+	TConstArrayView<FString> Keywords,
+	TConstArrayView<uint8> QuantizedCodes)
+{
+	const int64 ID = GetAssetIndexID(Asset);
+
+	// Lazy-init vector index from embedding dimension
+	if (Embedding.Num() > 0)
+	{
+		EnsureInitialized(Embedding.Num());
+	}
+
+	// Add to vector index — use pre-quantized path if available
+	if (VectorIndex && VectorIndex->IsTrained() && Embedding.Num() > 0)
+	{
+		if (QuantizedCodes.Num() > 0 && VectorIndex->SupportsQuantization())
+		{
+			VectorIndex->AddQuantized(ID, QuantizedCodes);
+		}
+		else
+		{
+			VectorIndex->Add(ID, Embedding);
+		}
+	}
+
+	// Add to BM25 index
+	if (IsHybridEnabled())
+	{
+		BM25Index->Add(ID, Asset.GetObjectPathString(), Caption, Keywords);
+	}
+}
+
+void FHybridSearchIndex::Remove(const FAssetData& Asset)
+{
+	Remove(Asset.GetObjectPathString());
+}
+
+void FHybridSearchIndex::Remove(FStringView AssetPath)
+{
+	const int64 ID = GetAssetIndexID(AssetPath);
+
+	if (VectorIndex)
+	{
+		VectorIndex->Remove(ID);
+	}
+	BM25Index->Remove(ID);
+}
+
+bool FHybridSearchIndex::Contains(const FAssetData& Asset) const
+{
+	const int64 ID = GetAssetIndexID(Asset);
+	const bool bInVector = VectorIndex && VectorIndex->Contains(ID);
+
+	if (!IsHybridEnabled())
+	{
+		return bInVector;
+	}
+
+	return bInVector && BM25Index->Contains(ID);
+}
+
+bool FHybridSearchIndex::Update(const FAssetData& NewAsset, FStringView OldAssetPath)
+{
+	const int64 OldID = GetAssetIndexID(OldAssetPath);
+	const int64 NewID = GetAssetIndexID(NewAsset);
+
+	bool bUpdated = false;
+
+	if (VectorIndex)
+	{
+		bUpdated = VectorIndex->Update(OldID, NewID);
+	}
+
+	if (IsHybridEnabled())
+	{
+		BM25Index->Update(OldID, OldAssetPath, NewAsset.GetObjectPathString());
+	}
+
+	return bUpdated;
+}
+
+TArray<FHybridSearchResult> FHybridSearchIndex::Search(
+	FStringView QueryText,
+	TConstArrayView<float> QueryEmbedding,
+	int32 K,
+	TConstArrayView<int64> IDFilter) const
+{
+	const USemanticSearchSettings* Settings = USemanticSearchSettings::Get();
+	const int32 RRFk = Settings ? Settings->RRFConstant : 60;
+	const int32 Oversample = Settings ? Settings->RRFOversample : 3;
+
+	// Fetch more candidates than K from each source for better fusion
+	const int32 FetchK = K * Oversample;
+
+	// Vector search
+	TArray<FSearchResult> VectorResults;
+	if (VectorIndex && QueryEmbedding.Num() > 0)
+	{
+		VectorResults = VectorIndex->Search(QueryEmbedding, FetchK, IDFilter);
+	}
+
+	// BM25 search
+	TArray<FBM25Result> TextResults;
+	if (IsHybridEnabled() && QueryText.Len() > 0)
+	{
+		TextResults = BM25Index->Search(QueryText, FetchK, IDFilter);
+	}
+
+	// If only one source has results, convert directly
+	if (VectorResults.Num() == 0 && TextResults.Num() == 0)
+	{
+		return {};
+	}
+
+	if (TextResults.Num() == 0)
+	{
+		// Vector only
+		TArray<FHybridSearchResult> Results;
+		Results.Reserve(FMath::Min(K, VectorResults.Num()));
+		for (int32 i = 0; i < VectorResults.Num() && i < K; ++i)
+		{
+			FHybridSearchResult R;
+			R.ID = VectorResults[i].ID;
+			R.VectorDistance = VectorResults[i].Distance;
+			R.RRFScore = 1.0f / (RRFk + i + 1);
+			Results.Add(R);
+		}
+		return Results;
+	}
+
+	if (VectorResults.Num() == 0)
+	{
+		// BM25 only
+		TArray<FHybridSearchResult> Results;
+		Results.Reserve(FMath::Min(K, TextResults.Num()));
+		for (int32 i = 0; i < TextResults.Num() && i < K; ++i)
+		{
+			FHybridSearchResult R;
+			R.ID = TextResults[i].ID;
+			R.BM25Score = TextResults[i].Score;
+			R.RRFScore = 1.0f / (RRFk + i + 1);
+			Results.Add(R);
+		}
+		return Results;
+	}
+
+	// RRF fusion: combine ranks from both lists
+	TMap<int64, FHybridSearchResult> Merged;
+	Merged.Reserve(FMath::Max(VectorResults.Num(), TextResults.Num()));
+
+	for (int32 i = 0; i < VectorResults.Num(); ++i)
+	{
+		FHybridSearchResult& R = Merged.FindOrAdd(VectorResults[i].ID);
+		R.ID = VectorResults[i].ID;
+		R.VectorDistance = VectorResults[i].Distance;
+		R.RRFScore += 1.0f / (RRFk + i + 1);
+	}
+
+	for (int32 i = 0; i < TextResults.Num(); ++i)
+	{
+		FHybridSearchResult& R = Merged.FindOrAdd(TextResults[i].ID);
+		R.ID = TextResults[i].ID;
+		R.BM25Score = TextResults[i].Score;
+		R.RRFScore += 1.0f / (RRFk + i + 1);
+	}
+
+	TArray<FHybridSearchResult> Results;
+	Merged.GenerateValueArray(Results);
+	Results.Sort([](const FHybridSearchResult& R1, const FHybridSearchResult& R2)
+	{
+		return R1.RRFScore > R2.RRFScore;
+	});
+
+	if (Results.Num() > K)
+	{
+		Results.SetNum(K);
+	}
+
+	return Results;
+}
+
+} // namespace UE::SemanticSearch
diff --git a/Engine/Plugins/Experimental/SemanticSearch/Source/SemanticSearch/Private/SemanticSearchDerivedDataBuilder.cpp b/Engine/Plugins/Experimental/SemanticSearch/Source/SemanticSearch/Private/SemanticSearchDerivedDataBuilder.cpp
--- a/Engine/Plugins/Experimental/SemanticSearch/Source/SemanticSearch/Private/SemanticSearchDerivedDataBuilder.cpp
+++ b/Engine/Plugins/Experimental/SemanticSearch/Source/SemanticSearch/Private/SemanticSearchDerivedDataBuilder.cpp
@@ -2,9 +2,8 @@
 
 #include "SemanticSearchDerivedDataBuilder.h"
 
-#include "AssetProcessorManager.h"
-#include "AssetProcessors/AssetProcessorUtils.h"
 #include "Containers/SharedString.h"
+#include "HybridSearchIndex.h"
 #include "ProfilingDebugging/CpuProfilerTrace.h"
 #include "Containers/StringFwd.h"
 #include "DerivedDataCache.h"
@@ -281,7 +280,7 @@ void FSemanticSearchBuildCacheTask::BuildDerivedData()
 
 											// Quantize the embedding if the active index supports it (PQ).
 											FQuantizedEmbeddingDerivedData QuantizedEmbeddingData;
-											if (TSharedPtr<IVectorIndex> VectorIndex = FAssetProcessorManager::Get().GetVectorIndex())
+											if (TSharedPtr<IVectorIndex> VectorIndex = FHybridSearchIndex::Get().GetVectorIndex())
 											{
 												if (VectorIndex->SupportsQuantization() && VectorIndex->IsTrained())
 												{
@@ -338,13 +337,6 @@ void FSemanticSearchBuildCacheTask::BuildDerivedData()
 														}
 													});
 
-												// Add the embedding to the vector index
-												if (TSharedPtr<IVectorIndex> VectorIndex = FAssetProcessorManager::Get().GetVectorIndex())
-												{
-													const int64 AssetID = GetAssetIndexID(*Asset);
-													VectorIndex->Add(AssetID, EmbeddingDerivedData.Embedding);
-												}
-
 												BroadcastQueuedRequests(MoveTemp(EmbeddingDerivedData), MoveTemp(CaptionDerivedData), FString());
 											}
 											else
diff --git a/Engine/Plugins/Experimental/SemanticSearch/Source/SemanticSearch/Private/SemanticSearchModule.cpp b/Engine/Plugins/Experimental/SemanticSearch/Source/SemanticSearch/Private/SemanticSearchModule.cpp
--- a/Engine/Plugins/Experimental/SemanticSearch/Source/SemanticSearch/Private/SemanticSearchModule.cpp
+++ b/Engine/Plugins/Experimental/SemanticSearch/Source/SemanticSearch/Private/SemanticSearchModule.cpp
@@ -4,7 +4,7 @@
 
 #include "AssetProcessorManager.h"
 #include "AssetProcessors/AssetProcessorUtils.h"
-#include "VectorIndex/FlatVectorIndex.h"
+#include "HybridSearchIndex.h"
 
 #include "Modules/ModuleManager.h"
 #include "RemoteEmbeddingProvider.h"
@@ -72,11 +72,6 @@ void FSemanticSearchModule::StartupModule()
 	RegisteredEmbeddingProvider = MakeShared<FRemoteEmbeddingProvider>();
 	RegisterDefaultAssetProcessors();
 
-	// Default to a flat index. Dimension is set when first embedding arrives.
-	// Callers can replace this with a PQ index via SetVectorIndex().
-	constexpr int32 DefaultDimension = 1536;
-	FAssetProcessorManager::Get().SetVectorIndex(MakeShared<FFlatVectorIndex>(DefaultDimension));
-
 	IAssetRegistry& AssetRegistry = FModuleManager::LoadModuleChecked<FAssetRegistryModule>(TEXT("AssetRegistry")).Get();
 	AssetAddedDelegateHandle = AssetRegistry.OnAssetAdded().AddRaw(
 		this, &FSemanticSearchModule::OnAssetAdded);
@@ -104,26 +99,30 @@ void FSemanticSearchModule::ShutdownModule()
 
 void FSemanticSearchModule::TriggerIndexing(const FAssetData& AssetData, bool bBuildOnMiss)
 {
-	const int64 AssetID = GetAssetIndexID(AssetData);
 	FAssetProcessorManager::Get().GetEmbeddingData(AssetData,
-		[AssetID, AssetPath = AssetData.GetObjectPathString(), bBuildOnMiss](FAssetEmbeddingResult&& Result, FString&& Error) mutable
+		[Asset = FAssetData(AssetData), bBuildOnMiss](FAssetEmbeddingResult&& EmbeddingResult, FString&& Error) mutable
 		{
-			if (!Error.IsEmpty() || Result.Embedding.Num() == 0)
+			if (!Error.IsEmpty() || EmbeddingResult.Embedding.IsEmpty())
 			{
-				// "Not cached" on a cache-only query is expected — only warn when a full build was attempted
 				if (bBuildOnMiss)
 				{
-					UE_LOG(LogSemanticSearch, Warning, TEXT("Indexing failed for %s: %s"), *AssetPath, Error.IsEmpty() ? TEXT("empty embedding") : *Error);
+					UE_LOG(LogSemanticSearch, Warning, TEXT("Indexing failed for %s: %s"),
+						*Asset.GetObjectPathString(),
+						Error.IsEmpty() ? TEXT("empty embedding") : *Error);
 				}
 				return;
 			}
-			if (TSharedPtr<IVectorIndex> VectorIndex = FAssetProcessorManager::Get().GetVectorIndex())
-			{
-				// Remove first so resaves and any BuildDerivedData-added entries don't create duplicates
-				VectorIndex->Remove(AssetID);
-				VectorIndex->Add(AssetID, Result.Embedding);
-				UE_LOG(LogSemanticSearch, Log, TEXT("Indexed %s (%d dims, index size: %lld)"), *AssetPath, Result.Embedding.Num(), VectorIndex->GetCount());
-			}
+
+			FAssetProcessorManager::Get().GetCaptionData(Asset,
+				[Asset, Embedding = MoveTemp(EmbeddingResult.Embedding)](FAssetCaptionResult&& CaptionResult, FString&& CaptionError) mutable
+				{
+					FHybridSearchIndex& Index = FHybridSearchIndex::Get();
+					Index.Remove(Asset);
+					Index.Add(Asset, Embedding, CaptionResult.Caption, CaptionResult.Keywords);
+					UE_LOG(LogSemanticSearch, Log, TEXT("Indexed %s (%d dims, index size: %lld)"),
+						*Asset.GetObjectPathString(), Embedding.Num(),
+						Index.GetVectorIndex() ? Index.GetVectorIndex()->GetCount() : 0);
+				});
 		}, bBuildOnMiss);
 }
 
@@ -159,47 +158,36 @@ void FSemanticSearchModule::OnAssetRenamed(const FAssetData& NewAssetData, const
 
 	if (!bOldInFolder && !bNewInFolder)
 	{
-		UE_LOG(LogSemanticSearch, Log, TEXT("OnAssetRenamed: neither path in indexed folder, ignoring"));
 		return;
 	}
 
-	TSharedPtr<IVectorIndex> VectorIndex = FAssetProcessorManager::Get().GetVectorIndex();
+	FHybridSearchIndex& Index = FHybridSearchIndex::Get();
 
-	const int64 OldID = GetAssetIndexID(OldObjectPath);
-	const int64 NewID = GetAssetIndexID(NewAssetData);
-
-	UE_LOG(LogSemanticSearch, Log, TEXT("OnAssetRenamed: OldID=%lld, NewID=%lld"), OldID, NewID);
-
-	// Try to reuse the existing embedding from the index to avoid a round-trip to the backend
-	TArray<float> Embedding;
-	if (bOldInFolder && VectorIndex && VectorIndex->TryGetEmbedding(OldID, Embedding))
+	if (bOldInFolder && bNewInFolder)
 	{
-		VectorIndex->Remove(OldID);
-		if (bNewInFolder)
+		if (Index.Update(NewAssetData, OldObjectPath))
 		{
-			VectorIndex->Add(NewID, Embedding);
-			UE_LOG(LogSemanticSearch, Log, TEXT("OnAssetRenamed: reused embedding (%d dims) for %s -> %s (index size: %lld)"),
-				Embedding.Num(), *OldObjectPath, *NewAssetData.GetObjectPathString(), VectorIndex->GetCount());
+			UE_LOG(LogSemanticSearch, Log, TEXT("OnAssetRenamed: updated index for %s -> %s"),
+				*OldObjectPath, *NewAssetData.GetObjectPathString());
 		}
 		else
 		{
-			UE_LOG(LogSemanticSearch, Log, TEXT("OnAssetRenamed: removed %s from index (moved out of indexed folder, index size: %lld)"),
-				*OldObjectPath, VectorIndex->GetCount());
+			UE_LOG(LogSemanticSearch, Log, TEXT("OnAssetRenamed: not in index, triggering full reindex for %s"),
+				*NewAssetData.GetObjectPathString());
+			TriggerIndexing(NewAssetData, CVarAutoIndexOnSave.GetValueOnGameThread());
 		}
 	}
+	else if (bOldInFolder)
+	{
+		UE_LOG(LogSemanticSearch, Log, TEXT("OnAssetRenamed: removed %s from index (moved out of indexed folder)"),
+			*OldObjectPath);
+		Index.Remove(OldObjectPath);
+	}
 	else
 	{
-		// Embedding not in index (e.g. index was rebuilt without this asset) — fall back to full reindex
-		if (bOldInFolder && VectorIndex)
-		{
-			const bool bRemoved = VectorIndex->Remove(OldID);
-			UE_LOG(LogSemanticSearch, Log, TEXT("OnAssetRenamed: embedding not in index for %s (Remove=%d)"), *OldObjectPath, bRemoved);
-		}
-		if (bNewInFolder)
-		{
-			UE_LOG(LogSemanticSearch, Log, TEXT("OnAssetRenamed: triggering full reindex for %s"), *NewAssetData.GetObjectPathString());
-			TriggerIndexing(NewAssetData, CVarAutoIndexOnSave.GetValueOnGameThread());
-		}
+		UE_LOG(LogSemanticSearch, Log, TEXT("OnAssetRenamed: moved into indexed folder, indexing %s"),
+			*NewAssetData.GetObjectPathString());
+		TriggerIndexing(NewAssetData, CVarAutoIndexOnSave.GetValueOnGameThread());
 	}
 }
 
@@ -207,10 +195,7 @@ void FSemanticSearchModule::OnAssetRemoved(const FAssetData& AssetData)
 {
 	if (IsInIndexedFolder(AssetData))
 	{
-		if (TSharedPtr<IVectorIndex> VectorIndex = FAssetProcessorManager::Get().GetVectorIndex())
-		{
-			VectorIndex->Remove(GetAssetIndexID(AssetData));
-		}
+		FHybridSearchIndex::Get().Remove(AssetData);
 	}
 }
 
@@ -224,11 +209,6 @@ void FSemanticSearchModule::RegisterEmbeddingProvider(TUniquePtr<IEmbeddingProvi
 	RegisteredEmbeddingProvider = MakeShareable<IEmbeddingProvider>(EmbeddingProvider.Release());
 }
 
-void FSemanticSearchModule::SetVectorIndex(TSharedPtr<IVectorIndex> InIndex)
-{
-	FAssetProcessorManager::Get().SetVectorIndex(MoveTemp(InIndex));
-}
-
 void FSemanticSearchModule::IndexAsset(const FAssetData& AssetData)
 {
 	TriggerIndexing(AssetData, /*bBuildOnMiss=*/true);
diff --git a/Engine/Plugins/Experimental/SemanticSearch/Source/SemanticSearch/Private/SemanticSearchModule.h b/Engine/Plugins/Experimental/SemanticSearch/Source/SemanticSearch/Private/SemanticSearchModule.h
--- a/Engine/Plugins/Experimental/SemanticSearch/Source/SemanticSearch/Private/SemanticSearchModule.h
+++ b/Engine/Plugins/Experimental/SemanticSearch/Source/SemanticSearch/Private/SemanticSearchModule.h
@@ -24,7 +24,6 @@ public:
 	virtual void ShutdownModule() override;
 
 	virtual void RegisterEmbeddingProvider(TUniquePtr<IEmbeddingProvider>&& EmbeddingProvider) override;
-	virtual void SetVectorIndex(TSharedPtr<IVectorIndex> InIndex) override;
 	virtual void IndexAsset(const FAssetData& AssetData) override;
 
 private:
diff --git a/Engine/Plugins/Experimental/SemanticSearch/Source/SemanticSearch/Private/SemanticSearchTestCommands.cpp b/Engine/Plugins/Experimental/SemanticSearch/Source/SemanticSearch/Private/SemanticSearchTestCommands.cpp
--- a/Engine/Plugins/Experimental/SemanticSearch/Source/SemanticSearch/Private/SemanticSearchTestCommands.cpp
+++ b/Engine/Plugins/Experimental/SemanticSearch/Source/SemanticSearch/Private/SemanticSearchTestCommands.cpp
@@ -1,6 +1,7 @@
 // Copyright Epic Games, Inc. All Rights Reserved.
 
 #include "SemanticSearchModule.h"
+#include "HybridSearchIndex.h"
 #include "AssetProcessorManager.h"
 #include "AssetProcessors/AssetProcessorUtils.h"
 #include "ISemanticSearchModule.h"
@@ -56,7 +57,9 @@ static FAutoConsoleCommand IndexStatsCommand(
 	TEXT("Print the number of assets currently in the vector index."),
 	FConsoleCommandDelegate::CreateLambda([]()
 	{
-		TSharedPtr<IVectorIndex> VectorIndex = FAssetProcessorManager::Get().GetVectorIndex();
+		FHybridSearchIndex& Searcher = FHybridSearchIndex::Get();
+		Searcher.EnsureInitialized(1536);
+		TSharedPtr<IVectorIndex> VectorIndex = Searcher.GetVectorIndex();
 		if (!VectorIndex)
 		{
 			UE_LOG(LogSemanticSearch, Warning, TEXT("No vector index is registered."));
@@ -69,7 +72,7 @@ static FAutoConsoleCommand IndexStatsCommand(
 		int32 Indexed = 0;
 		for (const FAssetData& Asset : Assets)
 		{
-			if (VectorIndex->Contains(GetAssetIndexID(Asset)))
+			if (Searcher.Contains(Asset))
 			{
 				++Indexed;
 			}
@@ -86,7 +89,9 @@ static FAutoConsoleCommand IndexAllCommand(
 	TEXT("Index all assets in the configured folders that are not yet in the FAISS index."),
 	FConsoleCommandDelegate::CreateLambda([]()
 	{
-		TSharedPtr<IVectorIndex> VectorIndex = FAssetProcessorManager::Get().GetVectorIndex();
+		FHybridSearchIndex& Searcher = FHybridSearchIndex::Get();
+		Searcher.EnsureInitialized(1536);
+		TSharedPtr<IVectorIndex> VectorIndex = Searcher.GetVectorIndex();
 		if (!VectorIndex)
 		{
 			UE_LOG(LogSemanticSearch, Warning, TEXT("No vector index is registered."));
@@ -101,7 +106,7 @@ static FAutoConsoleCommand IndexAllCommand(
 		int32 AlreadyIndexed = 0;
 		for (const FAssetData& Asset : Assets)
 		{
-			if (VectorIndex->Contains(GetAssetIndexID(Asset)))
+			if (Searcher.Contains(Asset))
 			{
 				++AlreadyIndexed;
 			}
@@ -116,52 +121,49 @@ static FAutoConsoleCommand IndexAllCommand(
 	})
 );
 
-// SemanticSearch.TestSearch <query text>
+// SemanticSearch.TestSearch <query text> [K]
 static FAutoConsoleCommand TestSearchCommand(
 	TEXT("SemanticSearch.TestSearch"),
-	TEXT("Search the vector index with a text query. Usage: SemanticSearch.TestSearch <query text>"),
+	TEXT("Search the hybrid index. Usage: SemanticSearch.TestSearch <query text> [K=10]"),
 	FConsoleCommandWithArgsDelegate::CreateLambda([](const TArray<FString>& Args)
 	{
 		if (Args.Num() < 1)
 		{
-			UE_LOG(LogSemanticSearch, Warning, TEXT("Usage: SemanticSearch.TestSearch <query text>"));
+			UE_LOG(LogSemanticSearch, Warning, TEXT("Usage: SemanticSearch.TestSearch <query text> [K=10]"));
 			return;
 		}
 
-		TSharedPtr<IVectorIndex> VectorIndex = FAssetProcessorManager::Get().GetVectorIndex();
-		if (!VectorIndex || VectorIndex->GetCount() == 0)
+		// If the last arg is a number, treat it as K
+		int32 K = 10;
+		int32 QueryArgCount = Args.Num();
+		if (Args.Num() >= 2)
 		{
-			UE_LOG(LogSemanticSearch, Warning, TEXT("Vector index is empty — no assets indexed yet."));
-			return;
+			const int32 MaybeK = FCString::Atoi(*Args.Last());
+			if (MaybeK > 0 && Args.Last() == FString::FromInt(MaybeK))
+			{
+				K = MaybeK;
+				QueryArgCount = Args.Num() - 1;
+			}
 		}
 
-		const FString QueryText = FString::Join(Args, TEXT(" "));
-		UE_LOG(LogSemanticSearch, Log, TEXT("Searching for: \"%s\" (index size: %lld)"), *QueryText, VectorIndex->GetCount());
+		FHybridSearchIndex& Searcher = FHybridSearchIndex::Get();
+		Searcher.EnsureInitialized(1536);
+
+		const FString QueryText = FString::Join(TConstArrayView<FString>(Args).Left(QueryArgCount), TEXT(" "));
+		UE_LOG(LogSemanticSearch, Log, TEXT("Searching for: \"%s\" (K=%d)"), *QueryText, K);
 
 		// Build ID -> asset path map from all configured indexed folders
 		TMap<int64, FString> IDToPath;
 		IAssetRegistry& AssetRegistry = FModuleManager::LoadModuleChecked<FAssetRegistryModule>(TEXT("AssetRegistry")).Get();
-		TArray<FString> IndexedPaths;
-		if (IConsoleVariable* CVar = IConsoleManager::Get().FindConsoleVariable(TEXT("SemanticSearch.IndexedPaths")))
-		{
-			CVar->GetString().ParseIntoArray(IndexedPaths, TEXT(","), /*bCullEmpty=*/true);
-		}
-		for (FString& Folder : IndexedPaths)
+		TArray<FAssetData> Assets = GetAssetsInIndexedFolders(AssetRegistry);
+		for (const FAssetData& AssetData : Assets)
 		{
-			Folder.TrimStartAndEndInline();
-			// GetAssetsByPath expects the path without trailing slash
-			TArray<FAssetData> Assets;
-			FString FolderPath = Folder.EndsWith(TEXT("/")) ? Folder.LeftChop(1) : Folder;
-			AssetRegistry.GetAssetsByPath(FName(FolderPath), Assets, /*bRecursive=*/true);
-			for (const FAssetData& AssetData : Assets)
-			{
-				IDToPath.Add(GetAssetIndexID(AssetData), AssetData.GetObjectPathString());
-			}
+			IDToPath.Add(GetAssetIndexID(AssetData), AssetData.GetObjectPathString());
 		}
 
 		TSharedPtr<FRemoteEmbeddingProvider> Provider = GetOrCreateProvider();
 		Provider->GenerateEmbeddingAsync(QueryText,
-			[VectorIndex, IDToPath = MoveTemp(IDToPath)](const FEmbeddingResponse& Response)
+			[&Searcher, K, QueryText, IDToPath = MoveTemp(IDToPath)](const FEmbeddingResponse& Response)
 			{
 				if (!Response.ErrorMessage.IsEmpty())
 				{
@@ -169,16 +171,17 @@ static FAutoConsoleCommand TestSearchCommand(
 					return;
 				}
 
-				constexpr int32 K = 5;
-				TArray<FSearchResult> Results = VectorIndex->Search(Response.Embedding, K);
+				TArray<FHybridSearchResult> Results = Searcher.Search(QueryText, Response.Embedding, K);
 
 				UE_LOG(LogSemanticSearch, Log, TEXT("Top %d results:"), Results.Num());
 				for (int32 i = 0; i < Results.Num(); ++i)
 				{
 					const FString* Path = IDToPath.Find(Results[i].ID);
-					UE_LOG(LogSemanticSearch, Log, TEXT("  [%d] dist=%.4f  %s"),
+					UE_LOG(LogSemanticSearch, Log, TEXT("  [%d] rrf=%.4f  vec_dist=%.4f  bm25=%.4f  %s"),
 						i + 1,
-						Results[i].Distance,
+						Results[i].RRFScore,
+						Results[i].VectorDistance,
+						Results[i].BM25Score,
 						Path ? **Path : TEXT("<unknown>"));
 				}
 			}, nullptr);
diff --git a/Engine/Plugins/Experimental/SemanticSearch/Source/SemanticSearch/Private/TextIndex/BM25Index.cpp b/Engine/Plugins/Experimental/SemanticSearch/Source/SemanticSearch/Private/TextIndex/BM25Index.cpp
new file mode 100644
--- /dev/null
+++ b/Engine/Plugins/Experimental/SemanticSearch/Source/SemanticSearch/Private/TextIndex/BM25Index.cpp
@@ -0,0 +1,447 @@
+// Copyright Epic Games, Inc. All Rights Reserved.
+
+#include "TextIndex/BM25Index.h"
+
+#include "Misc/Crc.h"
+#include "Misc/Char.h"
+#include "Misc/ScopeRWLock.h"
+
+namespace UE::SemanticSearch
+{
+
+namespace Tokenizer
+{
+
+uint32 HashToken(FStringView Token)
+{
+	return FCrc::StrCrc32(Token.GetData(), Token.Len());
+}
+
+const TSet<uint32>& GetStopWordHashes()
+{
+	static TSet<uint32> Hashes;
+	if (Hashes.Num() == 0)
+	{
+		static const TCHAR* StopWords[] = {
+			TEXT("the"), TEXT("is"), TEXT("in"), TEXT("on"), TEXT("at"),
+			TEXT("to"), TEXT("a"), TEXT("an"), TEXT("and"), TEXT("or"),
+			TEXT("of"), TEXT("for"), TEXT("with"), TEXT("this"), TEXT("that"),
+			TEXT("by"), TEXT("be"), TEXT("was"), TEXT("were"), TEXT("it"),
+			TEXT("as"), TEXT("from"), TEXT("but"), TEXT("if"), TEXT("then"),
+			TEXT("out"), TEXT("so"), TEXT("such"), TEXT("when"), TEXT("which"),
+			TEXT("while"), TEXT("are"), TEXT("has"), TEXT("have"), TEXT("had"),
+			TEXT("not"), TEXT("no"), TEXT("can"), TEXT("will"), TEXT("do"),
+		};
+		for (const TCHAR* Word : StopWords)
+		{
+			Hashes.Add(HashToken(FStringView(Word)));
+		}
+	}
+	return Hashes;
+}
+
+// Split a single part on CamelCase boundaries, output lowercase token hashes.
+void SplitCamelCase(FStringView Part, TArray<uint32>& OutHashes, const TSet<uint32>& StopHashes)
+{
+	const int32 Len = Part.Len();
+	if (Len == 0)
+	{
+		return;
+	}
+
+	const TCHAR* Data = Part.GetData();
+	int32 Start = 0;
+
+	auto EmitToken = [Data, &OutHashes, &StopHashes](int32 TokenStart, int32 TokenEnd)
+	{
+		if (TokenEnd <= TokenStart)
+		{
+			return;
+		}
+		// Build lowercase token and hash it
+		FString Lower;
+		Lower.Reserve(TokenEnd - TokenStart);
+		for (int32 j = TokenStart; j < TokenEnd; ++j)
+		{
+			Lower.AppendChar(FChar::ToLower(Data[j]));
+		}
+		uint32 Hash = HashToken(Lower);
+		if (!StopHashes.Contains(Hash))
+		{
+			OutHashes.Add(Hash);
+		}
+	};
+
+	for (int32 i = 1; i <= Len; ++i)
+	{
+		bool bSplit = (i == Len);
+
+		if (!bSplit)
+		{
+			const TCHAR Prev = Data[i - 1];
+			const TCHAR Curr = Data[i];
+
+			// Upper→Lower with multi-char prefix: "GRass" → "G" | "Rass"
+			if (FChar::IsUpper(Prev) && FChar::IsLower(Curr) && (i - Start) > 1)
+			{
+				EmitToken(Start, i - 1);
+				Start = i - 1;
+				continue;
+			}
+			// Lower→Upper: "grass" | "Texture"
+			if (FChar::IsLower(Prev) && FChar::IsUpper(Curr))
+			{
+				bSplit = true;
+			}
+			// Letter↔Digit transitions
+			if (FChar::IsDigit(Prev) != FChar::IsDigit(Curr))
+			{
+				bSplit = true;
+			}
+		}
+
+		if (bSplit && i > Start)
+		{
+			EmitToken(Start, i);
+			Start = i;
+		}
+	}
+}
+
+TArray<uint32> Tokenize(FStringView Text)
+{
+	TArray<uint32> TokenHashes;
+	const TSet<uint32>& StopHashes = GetStopWordHashes();
+
+	const int32 Len = Text.Len();
+	const TCHAR* Data = Text.GetData();
+	int32 Start = -1;
+
+	// Phase 1: split on non-alphanumeric (underscore is separator)
+	for (int32 i = 0; i <= Len; ++i)
+	{
+		bool bIsWordChar = (i < Len) && FChar::IsAlnum(Data[i]);
+
+		if (bIsWordChar)
+		{
+			if (Start < 0)
+			{
+				Start = i;
+			}
+		}
+		else if (Start >= 0)
+		{
+			// Phase 2: CamelCase split this part
+			SplitCamelCase(FStringView(Data + Start, i - Start), TokenHashes, StopHashes);
+			Start = -1;
+		}
+	}
+
+	return TokenHashes;
+}
+
+} // namespace Tokenizer
+
+void FBM25Index::Add(int64 ID, FStringView AssetPath, FStringView Caption, TConstArrayView<FString> Keywords)
+{
+	// Tokenize all fields
+	TArray<uint32> AllTokens;
+	AllTokens.Append(Tokenizer::Tokenize(AssetPath));
+	AllTokens.Append(Tokenizer::Tokenize(Caption));
+	for (const FString& Keyword : Keywords)
+	{
+		AllTokens.Append(Tokenizer::Tokenize(Keyword));
+	}
+
+	if (AllTokens.Num() == 0)
+	{
+		return;
+	}
+
+	// Build term frequency map
+	TMap<uint32, uint8> TermFreqs;
+	for (uint32 Hash : AllTokens)
+	{
+		uint8& Freq = TermFreqs.FindOrAdd(Hash, 0);
+		if (Freq < 255)
+		{
+			Freq++;
+		}
+	}
+
+	// Track unique token hashes for this document (for removal)
+	TArray<uint32> UniqueTokens;
+	TermFreqs.GetKeys(UniqueTokens);
+
+	FWriteScopeLock WriteLock(IndexLock);
+
+	// Remove existing document if present (update case)
+	if (DocLengths.Contains(ID))
+	{
+		// Unlock and re-lock not needed since we hold write lock
+		// Inline removal logic to avoid recursive lock
+		if (TArray<uint32>* ExistingTokens = DocTokens.Find(ID))
+		{
+			for (uint32 Hash : *ExistingTokens)
+			{
+				if (TArray<FPostingEntry>* Postings = PostingLists.Find(Hash))
+				{
+					Postings->RemoveAll([ID](const FPostingEntry& E) { return E.DocID == ID; });
+					if (Postings->Num() == 0)
+					{
+						PostingLists.Remove(Hash);
+					}
+				}
+			}
+		}
+		if (uint16* OldLen = DocLengths.Find(ID))
+		{
+			TotalTokenCount -= *OldLen;
+		}
+		DocLengths.Remove(ID);
+		DocTokens.Remove(ID);
+	}
+
+	// Add to posting lists
+	for (const auto& Pair : TermFreqs)
+	{
+		FPostingEntry Entry;
+		Entry.DocID = ID;
+		Entry.TermFreq = Pair.Value;
+		PostingLists.FindOrAdd(Pair.Key).Add(Entry);
+	}
+
+	const uint16 DocLen = static_cast<uint16>(FMath::Min(AllTokens.Num(), (int32)MAX_uint16));
+	DocLengths.Add(ID, DocLen);
+	DocTokens.Add(ID, MoveTemp(UniqueTokens));
+	TotalTokenCount += DocLen;
+}
+
+void FBM25Index::Remove(int64 ID)
+{
+	FWriteScopeLock WriteLock(IndexLock);
+
+	TArray<uint32>* Tokens = DocTokens.Find(ID);
+	if (!Tokens)
+	{
+		return;
+	}
+
+	for (uint32 Hash : *Tokens)
+	{
+		if (TArray<FPostingEntry>* Postings = PostingLists.Find(Hash))
+		{
+			Postings->RemoveAll([ID](const FPostingEntry& E) { return E.DocID == ID; });
+			if (Postings->Num() == 0)
+			{
+				PostingLists.Remove(Hash);
+			}
+		}
+	}
+
+	if (uint16* Len = DocLengths.Find(ID))
+	{
+		TotalTokenCount -= *Len;
+	}
+	DocLengths.Remove(ID);
+	DocTokens.Remove(ID);
+}
+
+TArray<FBM25Result> FBM25Index::Search(FStringView Query, int32 K, TConstArrayView<int64> IDFilter) const
+{
+	TArray<uint32> QueryTokenHashes = Tokenizer::Tokenize(Query);
+	if (QueryTokenHashes.Num() == 0)
+	{
+		return {};
+	}
+
+	FReadScopeLock ReadLock(IndexLock);
+
+	const int32 NumDocs = DocLengths.Num();
+	if (NumDocs == 0)
+	{
+		return {};
+	}
+
+	// Build filter set if provided
+	TSet<int64> FilterSet;
+	if (IDFilter.Num() > 0)
+	{
+		FilterSet.Reserve(IDFilter.Num());
+		for (int64 ID : IDFilter)
+		{
+			FilterSet.Add(ID);
+		}
+	}
+	const bool bHasFilter = FilterSet.Num() > 0;
+
+	// Collect candidate documents: union of posting lists for query tokens
+	TSet<int64> Candidates;
+	for (uint32 Hash : QueryTokenHashes)
+	{
+		if (const TArray<FPostingEntry>* Postings = PostingLists.Find(Hash))
+		{
+			for (const FPostingEntry& Entry : *Postings)
+			{
+				if (!bHasFilter || FilterSet.Contains(Entry.DocID))
+				{
+					Candidates.Add(Entry.DocID);
+				}
+			}
+		}
+	}
+
+	// Score candidates
+	TArray<FBM25Result> Results;
+	Results.Reserve(FMath::Min(Candidates.Num(), K));
+
+	for (int64 DocID : Candidates)
+	{
+		float Score = ScoreDocument(DocID, QueryTokenHashes);
+		if (Score > 0.0f)
+		{
+			FBM25Result Result;
+			Result.ID = DocID;
+			Result.Score = Score;
+			Results.Add(Result);
+		}
+	}
+
+	// Sort descending by score
+	Results.Sort([](const FBM25Result& R1, const FBM25Result& R2) { return R1.Score > R2.Score; });
+
+	if (Results.Num() > K)
+	{
+		Results.SetNum(K);
+	}
+
+	return Results;
+}
+
+float FBM25Index::ScoreDocument(int64 DocID, TConstArrayView<uint32> QueryTokenHashes) const
+{
+	const uint16* DocLenPtr = DocLengths.Find(DocID);
+	if (!DocLenPtr)
+	{
+		return 0.0f;
+	}
+
+	const float DL = static_cast<float>(*DocLenPtr);
+	const int32 NumDocs = DocLengths.Num();
+	const float AvgDL = (NumDocs > 0) ? static_cast<float>(TotalTokenCount) / NumDocs : 1.0f;
+
+	float Score = 0.0f;
+
+	for (uint32 Hash : QueryTokenHashes)
+	{
+		const TArray<FPostingEntry>* Postings = PostingLists.Find(Hash);
+		if (!Postings)
+		{
+			continue;
+		}
+
+		// Find this document's TF in the posting list
+		uint8 TF = 0;
+		for (const FPostingEntry& Entry : *Postings)
+		{
+			if (Entry.DocID == DocID)
+			{
+				TF = Entry.TermFreq;
+				break;
+			}
+		}
+		if (TF == 0)
+		{
+			continue;
+		}
+
+		// IDF: log((N - df + 0.5) / (df + 0.5) + 1)
+		const float Df = static_cast<float>(Postings->Num());
+		const float N = static_cast<float>(NumDocs);
+		const float IDFScore = FMath::Loge((N - Df + 0.5f) / (Df + 0.5f) + 1.0f);
+
+		// BM25 score for this term
+		const float TermFreq = static_cast<float>(TF);
+		const float Numerator = IDFScore * TermFreq * (K1 + 1.0f);
+		const float Denominator = TermFreq + K1 * (1.0f - B + B * (DL / AvgDL));
+		Score += Numerator / Denominator;
+	}
+
+	return Score;
+}
+
+int64 FBM25Index::GetCount() const
+{
+	FReadScopeLock ReadLock(IndexLock);
+	return DocLengths.Num();
+}
+
+bool FBM25Index::Contains(int64 ID) const
+{
+	FReadScopeLock ReadLock(IndexLock);
+	return DocLengths.Contains(ID);
+}
+
+bool FBM25Index::Update(int64 ID, FStringView OldPath, FStringView NewPath)
+{
+	TArray<uint32> OldPathTokens = Tokenizer::Tokenize(OldPath);
+	TArray<uint32> NewPathTokens = Tokenizer::Tokenize(NewPath);
+
+	FWriteScopeLock WriteLock(IndexLock);
+
+	if (!DocLengths.Contains(ID))
+	{
+		return false;
+	}
+
+	TArray<uint32>* DocTokenList = DocTokens.Find(ID);
+	if (!DocTokenList)
+	{
+		return false;
+	}
+
+	// Remove old path token entries from posting lists
+	for (uint32 Hash : OldPathTokens)
+	{
+		if (TArray<FPostingEntry>* Postings = PostingLists.Find(Hash))
+		{
+			Postings->RemoveAll([ID](const FPostingEntry& E) { return E.DocID == ID; });
+			if (Postings->Num() == 0)
+			{
+				PostingLists.Remove(Hash);
+			}
+		}
+		DocTokenList->Remove(Hash);
+	}
+
+	// Build term frequency map for new path tokens
+	TMap<uint32, uint8> NewTermFreqs;
+	for (uint32 Hash : NewPathTokens)
+	{
+		uint8& Freq = NewTermFreqs.FindOrAdd(Hash, 0);
+		if (Freq < 255)
+		{
+			Freq++;
+		}
+	}
+
+	// Add new path token entries to posting lists
+	for (const auto& Pair : NewTermFreqs)
+	{
+		FPostingEntry Entry;
+		Entry.DocID = ID;
+		Entry.TermFreq = Pair.Value;
+		PostingLists.FindOrAdd(Pair.Key).Add(Entry);
+		DocTokenList->AddUnique(Pair.Key);
+	}
+
+	// Update doc length
+	const int32 LenDelta = NewPathTokens.Num() - OldPathTokens.Num();
+	uint16* DocLen = DocLengths.Find(ID);
+	*DocLen = static_cast<uint16>(FMath::Clamp(static_cast<int32>(*DocLen) + LenDelta, 0, (int32)MAX_uint16));
+	TotalTokenCount += LenDelta;
+
+	return true;
+}
+
+} // namespace UE::SemanticSearch
diff --git a/Engine/Plugins/Experimental/SemanticSearch/Source/SemanticSearch/Private/TextIndex/BM25Index.h b/Engine/Plugins/Experimental/SemanticSearch/Source/SemanticSearch/Private/TextIndex/BM25Index.h
new file mode 100644
--- /dev/null
+++ b/Engine/Plugins/Experimental/SemanticSearch/Source/SemanticSearch/Private/TextIndex/BM25Index.h
@@ -0,0 +1,75 @@
+// Copyright Epic Games, Inc. All Rights Reserved.
+
+#pragma once
+
+#include "Containers/Array.h"
+#include "Containers/ArrayView.h"
+#include "Containers/Map.h"
+#include "Containers/StringView.h"
+#include "HAL/CriticalSection.h"
+
+namespace UE::SemanticSearch
+{
+
+struct FBM25Result
+{
+	int64 ID = -1;
+	float Score = 0.0f;
+};
+
+/**
+ * Memory-efficient BM25 inverted index for text search over asset captions,
+ * keywords, and paths.
+ */
+class FBM25Index
+{
+public:
+	FBM25Index() = default;
+
+	/** Add a document to the index. Tokenizes path, caption, and keywords. */
+	void Add(int64 ID, FStringView AssetPath, FStringView Caption, TConstArrayView<FString> Keywords);
+
+	/** Remove a document from the index. */
+	void Remove(int64 ID);
+
+	/** Search for documents matching the query. Returns top K results sorted by score descending. */
+	TArray<FBM25Result> Search(FStringView Query, int32 K, TConstArrayView<int64> IDFilter = {}) const;
+
+	/** Number of documents in the index. */
+	int64 GetCount() const;
+
+	/** Check if a document is already indexed. */
+	bool Contains(int64 ID) const;
+
+	/** Update a document's path tokens without touching caption/keyword tokens. */
+	bool Update(int64 ID, FStringView OldPath, FStringView NewPath);
+
+private:
+	struct FPostingEntry
+	{
+		int64 DocID;
+		uint8 TermFreq;  // capped at 255
+	};
+
+	/** Score a single document against query tokens. */
+	float ScoreDocument(int64 DocID, TConstArrayView<uint32> QueryTokenHashes) const;
+
+	static constexpr float K1 = 1.5f;
+	static constexpr float B = 0.75f;
+
+	// Inverted index: token hash → posting list
+	TMap<uint32, TArray<FPostingEntry>> PostingLists;
+
+	// Per-document: doc length (token count)
+	TMap<int64, uint16> DocLengths;
+
+	// Per-document: which token hashes it contains (for removal)
+	TMap<int64, TArray<uint32>> DocTokens;
+
+	// Global stats
+	int64 TotalTokenCount = 0;
+
+	mutable FRWLock IndexLock;
+};
+
+} // namespace UE::SemanticSearch
diff --git a/Engine/Plugins/Experimental/SemanticSearch/Source/SemanticSearch/Public/AssetProcessorManager.h b/Engine/Plugins/Experimental/SemanticSearch/Source/SemanticSearch/Public/AssetProcessorManager.h
--- a/Engine/Plugins/Experimental/SemanticSearch/Source/SemanticSearch/Public/AssetProcessorManager.h
+++ b/Engine/Plugins/Experimental/SemanticSearch/Source/SemanticSearch/Public/AssetProcessorManager.h
@@ -10,7 +10,6 @@
 #include "DerivedDataCacheKey.h"
 #include "HAL/CriticalSection.h"
 #include "Interfaces/IAssetProcessor.h"
-#include "Interfaces/IVectorIndex.h"
 #include "Templates/Function.h"
 #include "Templates/UniquePtr.h"
 #include "Templates/UnrealTypeTraits.h"
@@ -45,9 +44,6 @@ public:
 	SEMANTICSEARCH_API bool GetCaptionData(const FAssetData& InAssetData, TUniqueFunction<void(FAssetCaptionResult&&, FString&&)>&& InOnDataAvailable);
 	SEMANTICSEARCH_API bool GetEmbeddingData(const FAssetData& InAssetData, TUniqueFunction<void(FAssetEmbeddingResult&&, FString&&)>&& InOnDataAvailable, bool bBuildOnMiss = true);
 
-	SEMANTICSEARCH_API void SetVectorIndex(TSharedPtr<IVectorIndex> InIndex);
-	SEMANTICSEARCH_API TSharedPtr<IVectorIndex> GetVectorIndex() const;
-
 	template<typename FAssetProcessor>
 	void RegisterAssetProcessor();
 
@@ -78,8 +74,6 @@ private:
 	TMap<FTopLevelAssetPath, TSharedPtr<IAssetProcessor>> AssetTypeToProcessor;
 
 	TStripedMap<32, DerivedData::FCacheKey, TSharedRef<Private::FSemanticSearchBuildCacheTask>> BuildsJobs;
-
-	TSharedPtr<IVectorIndex> VectorIndex;
 };
 
 template<typename FAssetProcessor>
diff --git a/Engine/Plugins/Experimental/SemanticSearch/Source/SemanticSearch/Public/HybridSearchIndex.h b/Engine/Plugins/Experimental/SemanticSearch/Source/SemanticSearch/Public/HybridSearchIndex.h
new file mode 100644
--- /dev/null
+++ b/Engine/Plugins/Experimental/SemanticSearch/Source/SemanticSearch/Public/HybridSearchIndex.h
@@ -0,0 +1,99 @@
+// Copyright Epic Games, Inc. All Rights Reserved.
+
+#pragma once
+
+#include "Containers/Array.h"
+#include "Containers/ArrayView.h"
+#include "Containers/StringView.h"
+#include "Interfaces/IVectorIndex.h"
+#include "Templates/SharedPointer.h"
+#include "Templates/UniquePtr.h"
+#include "HAL/CriticalSection.h"
+
+struct FAssetData;
+
+namespace UE::SemanticSearch
+{
+
+struct FHybridSearchResult
+{
+	int64 ID = -1;
+	float RRFScore = 0.0f;
+	float VectorDistance = -1.0f;	// -1 if not in vector results
+	float BM25Score = -1.0f;		// -1 if not in BM25 results
+};
+
+class FBM25Index;
+
+/**
+ * Singleton that owns all search indices (vector + BM25) and fuses results
+ * using Reciprocal Rank Fusion (RRF).
+ */
+class FHybridSearchIndex
+{
+public:
+	SEMANTICSEARCH_API static FHybridSearchIndex& Get();
+
+	/** Get the vector index (may be null before first Add). */
+	SEMANTICSEARCH_API TSharedPtr<IVectorIndex> GetVectorIndex() const;
+
+	/** Create the vector index if not yet initialized. */
+	void EnsureInitialized(int32 EmbeddingDimension);
+
+	/**
+	 * Main search entry point. Fuses vector + BM25 results via RRF.
+	 * @param QueryText			Text query for BM25 search.
+	 * @param QueryEmbedding	Embedding for vector search. Can be empty if not yet available.
+	 * @param K					Number of results to return.
+	 * @param IDFilter			If non-empty, only search among these IDs.
+	 */
+	SEMANTICSEARCH_API TArray<FHybridSearchResult> Search(
+		FStringView QueryText,
+		TConstArrayView<float> QueryEmbedding,
+		int32 K,
+		TConstArrayView<int64> IDFilter = {}) const;
+
+	/**
+	 * Add a single asset to both indices. Lazily creates vector index on first call.
+	 * @param QuantizedCodes  If non-empty and index is PQ, inserts pre-quantized codes directly (skips re-encoding).
+	 */
+	SEMANTICSEARCH_API void Add(const FAssetData& Asset,
+		TConstArrayView<float> Embedding, FStringView Caption, TConstArrayView<FString> Keywords,
+		TConstArrayView<uint8> QuantizedCodes = {});
+
+	/** Remove a single asset from both indices. */
+	SEMANTICSEARCH_API void Remove(const FAssetData& Asset);
+
+	/** Remove by path string (for when the asset has already been deleted). */
+	SEMANTICSEARCH_API void Remove(FStringView AssetPath);
+
+	/** Check if an asset is in the vector index. */
+	SEMANTICSEARCH_API bool Contains(const FAssetData& Asset) const;
+
+	/** Update an asset's ID after a rename. Re-keys vector data and swaps path tokens in BM25. */
+	SEMANTICSEARCH_API bool Update(const FAssetData& NewAsset, FStringView OldAssetPath);
+
+	/** Whether BM25 hybrid search is enabled (from settings). */
+	bool IsHybridEnabled() const;
+
+	/** Whether the vector index has been initialized. */
+	bool IsInitialized() const { return bInitialized; }
+
+	/** Access the BM25 index directly (e.g. for diagnostics). */
+	FBM25Index& GetBM25Index();
+
+private:
+	FHybridSearchIndex();
+
+	bool bInitialized = false;
+	TSharedPtr<IVectorIndex> VectorIndex;
+	TUniquePtr<FBM25Index> BM25Index;
+
+	mutable FRWLock IndexLock;
+};
+
+/** Compute the stable int64 ID used for an asset in the indices. */
+SEMANTICSEARCH_API int64 GetAssetIndexID(const FAssetData& Asset);
+SEMANTICSEARCH_API int64 GetAssetIndexID(FStringView AssetPath);
+
+} // namespace UE::SemanticSearch
diff --git a/Engine/Plugins/Experimental/SemanticSearch/Source/SemanticSearch/Public/ISemanticSearchModule.h b/Engine/Plugins/Experimental/SemanticSearch/Source/SemanticSearch/Public/ISemanticSearchModule.h
--- a/Engine/Plugins/Experimental/SemanticSearch/Source/SemanticSearch/Public/ISemanticSearchModule.h
+++ b/Engine/Plugins/Experimental/SemanticSearch/Source/SemanticSearch/Public/ISemanticSearchModule.h
@@ -4,7 +4,6 @@
 
 #include "AssetRegistry/AssetData.h"
 #include "Interfaces/IEmbeddingProvider.h"
-#include "Interfaces/IVectorIndex.h"
 #include "Modules/ModuleInterface.h"
 #include "Modules/ModuleManager.h"
 #include "Templates/SharedPointer.h"
@@ -13,7 +12,7 @@
 namespace UE::SemanticSearch
 {
 	class ISemanticSearchModule : public IModuleInterface
-	{ 
+	{
 	public:
 		static ISemanticSearchModule& Get()
 		{
@@ -26,9 +25,6 @@ namespace UE::SemanticSearch
 		/** Returns the active embedding provider, or nullptr if none has been registered. */
 		virtual TSharedPtr<IEmbeddingProvider> GetEmbeddingProvider() const = 0;
 
-		/** Set the vector index used for similarity search. Replaces any existing index. */
-		virtual void SetVectorIndex(TSharedPtr<IVectorIndex> InIndex) = 0;
-
 		/** Index a single asset, always building from the backend if not already cached. */
 		virtual void IndexAsset(const FAssetData& AssetData) = 0;
 	};
diff --git a/Engine/Plugins/Experimental/SemanticSearch/Source/SemanticSearch/Public/Interfaces/IVectorIndex.h b/Engine/Plugins/Experimental/SemanticSearch/Source/SemanticSearch/Public/Interfaces/IVectorIndex.h
--- a/Engine/Plugins/Experimental/SemanticSearch/Source/SemanticSearch/Public/Interfaces/IVectorIndex.h
+++ b/Engine/Plugins/Experimental/SemanticSearch/Source/SemanticSearch/Public/Interfaces/IVectorIndex.h
@@ -17,24 +17,13 @@ struct FSearchResult
 
 /**
  * Abstract interface for vector similarity search indices.
- *
- * All mutation methods are batch-oriented (take arrays). Non-virtual inline
- * helpers provide single-item convenience wrappers.
- *
- * Quantization methods (Quantize, AddQuantized) check(SupportsQuantization())
- * and will crash if called on an index that doesn't support it. Callers must
- * check SupportsQuantization() before calling.
- *
- * Concrete implementations:
- *   - FFlatVectorIndex: Exact brute-force search (IndexFlatL2). No training needed.
- *   - FPQVectorIndex: Product Quantization compressed search. Requires training.
  */
 class IVectorIndex
 {
 public:
 	virtual ~IVectorIndex() = default;
 
-	// --- Virtual batch methods (implementations override these) ---
+	// --- Virtual batch methods ---
 
 	/** Train the index on representative vectors. No-op for flat indices. */
 	virtual void Train(TConstArrayView<float> Vectors, int64 NumVectors) = 0;
@@ -115,6 +104,22 @@ public:
 
 	/** Retrieve the embedding for a given ID. Returns false if not supported or ID not found. */
 	virtual bool TryGetEmbedding(int64 ID, TArray<float>& OutEmbedding) const { return false; }
+
+	/**
+	 * Re-key a vector from OldID to NewID, preserving the embedding data.
+	 * Returns true if the vector was found and updated, false if OldID was not in the index.
+	 */
+	virtual bool Update(int64 OldID, int64 NewID)
+	{
+		TArray<float> Embedding;
+		if (!TryGetEmbedding(OldID, Embedding))
+		{
+			return false;
+		}
+		Remove(OldID);
+		Add(NewID, Embedding);
+		return true;
+	}
 };
 
 } // namespace UE::SemanticSearch
diff --git a/Engine/Plugins/Experimental/SemanticSearch/Source/SemanticSearch/Public/Settings/SemanticSearchSettings.h b/Engine/Plugins/Experimental/SemanticSearch/Source/SemanticSearch/Public/Settings/SemanticSearchSettings.h
new file mode 100644
--- /dev/null
+++ b/Engine/Plugins/Experimental/SemanticSearch/Source/SemanticSearch/Public/Settings/SemanticSearchSettings.h
@@ -0,0 +1,47 @@
+// Copyright Epic Games, Inc. All Rights Reserved.
+
+#pragma once
+
+#include "Engine/DeveloperSettings.h"
+#include "SemanticSearchSettings.generated.h"
+
+UENUM()
+enum class ESemanticSearchIndexType : uint8
+{
+	Flat    UMETA(DisplayName = "Flat (Exact)"),
+	PQ      UMETA(DisplayName = "Product Quantization (Compressed)")
+};
+
+UCLASS(config = EditorPerProjectUserSettings, DefaultConfig, meta = (DisplayName = "Semantic Search"))
+class SEMANTICSEARCH_API USemanticSearchSettings : public UDeveloperSettings
+{
+	GENERATED_BODY()
+
+public:
+	static const USemanticSearchSettings* Get()
+	{
+		return GetDefault<USemanticSearchSettings>();
+	}
+
+	virtual FName GetContainerName() const override { return TEXT("Editor"); }
+	virtual FName GetCategoryName() const override { return TEXT("Plugins"); }
+	virtual FName GetSectionName() const override { return TEXT("Semantic Search"); }
+
+	UPROPERTY(EditAnywhere, Config, Category = "Vector Index")
+	ESemanticSearchIndexType IndexType = ESemanticSearchIndexType::Flat;
+
+	UPROPERTY(EditAnywhere, Config, Category = "Search")
+	bool bEnableHybridSearch = true;
+
+	UPROPERTY(EditAnywhere, Config, Category = "Search",
+		meta = (EditCondition = "bEnableHybridSearch", EditConditionHides,
+			ClampMin = "1", ClampMax = "1000",
+			ToolTip = "RRF constant (k). Higher values reduce the impact of rank differences."))
+	int32 RRFConstant = 60;
+
+	UPROPERTY(EditAnywhere, Config, Category = "Search",
+		meta = (EditCondition = "bEnableHybridSearch", EditConditionHides,
+			ClampMin = "1", ClampMax = "10",
+			ToolTip = "Oversample multiplier for RRF. Each source fetches K * this many candidates before fusion."))
+	int32 RRFOversample = 3;
+};
diff --git a/Engine/Plugins/Experimental/SemanticSearch/Source/SemanticSearch/SemanticSearch.Build.cs b/Engine/Plugins/Experimental/SemanticSearch/Source/SemanticSearch/SemanticSearch.Build.cs
--- a/Engine/Plugins/Experimental/SemanticSearch/Source/SemanticSearch/SemanticSearch.Build.cs
+++ b/Engine/Plugins/Experimental/SemanticSearch/Source/SemanticSearch/SemanticSearch.Build.cs
@@ -24,6 +24,7 @@ public class SemanticSearch : ModuleRules
 				"Json",
 				"PhysicsCore",
 				"UnrealEd",
+				"DeveloperSettings",
 				"FAISS"
 			});
 
diff --git a/Engine/Plugins/Experimental/SemanticSearch/Source/SemanticSearchEditorIntegrations/Private/ContentBrowser/SemanticSearchCB.cpp b/Engine/Plugins/Experimental/SemanticSearch/Source/SemanticSearchEditorIntegrations/Private/ContentBrowser/SemanticSearchCB.cpp
--- a/Engine/Plugins/Experimental/SemanticSearch/Source/SemanticSearchEditorIntegrations/Private/ContentBrowser/SemanticSearchCB.cpp
+++ b/Engine/Plugins/Experimental/SemanticSearch/Source/SemanticSearchEditorIntegrations/Private/ContentBrowser/SemanticSearchCB.cpp
@@ -2,22 +2,19 @@
 
 #include "ContentBrowser/SemanticSearchCB.h"
 
-#include "AssetProcessorManager.h"
 #include "AssetViewTypes.h"
 #include "Async/Async.h"
 #include "Async/TaskGraphInterfaces.h"
 #include "ContentBrowserModule.h"
 #include "Experimental/ContentBrowserSearchOverride.h"
+#include "HybridSearchIndex.h"
 #include "ISemanticSearchModule.h"
 #include "Input/Reply.h"
 #include "Interfaces/IEmbeddingProvider.h"
-#include "Math/UnrealMathUtility.h"
-#include "Math/VectorRegister.h"
 #include "Modules/ModuleManager.h"
 #include "Styling/AppStyle.h"
 #include "Styling/SlateColor.h"
 #include "Templates/SharedPointer.h"
-#include "Templates/UniquePtr.h"
 #include "Widgets/Input/SButton.h"
 #include "Widgets/Text/STextBlock.h"
 
@@ -30,58 +27,6 @@ namespace Private
 	using namespace UE::Editor::ContentBrowser::Extension;
 	using namespace UE::SemanticSearch;
 
-	static float CosineSimilarity(const TArray<float>& A, const TArray<float>& B)
-	{
-		if (A.Num() != B.Num() || A.IsEmpty())
-		{
-			return 0.f;
-		}
-
-		const int32 Num = A.Num();
-		const float* __restrict DataA = A.GetData();
-		const float* __restrict DataB = B.GetData();
-
-		// Process 4 floats per iteration using SIMD fused multiply-add.
-		VectorRegister4Float VecDot  = VectorZeroFloat();
-		VectorRegister4Float VecNormA = VectorZeroFloat();
-		VectorRegister4Float VecNormB = VectorZeroFloat();
-
-		const int32 NumVec = Num & ~3;
-		for (int32 i = 0; i < NumVec; i += 4)
-		{
-			const VectorRegister4Float VA = VectorLoad(DataA + i);
-			const VectorRegister4Float VB = VectorLoad(DataB + i);
-			VecDot   = VectorMultiplyAdd(VA, VB, VecDot);
-			VecNormA = VectorMultiplyAdd(VA, VA, VecNormA);
-			VecNormB = VectorMultiplyAdd(VB, VB, VecNormB);
-		}
-
-		// Horizontal sum: add pairs, then add the two pair results.
-		auto HSum = [](VectorRegister4Float V) -> float
-		{
-			VectorRegister4Float Shuf = VectorSwizzle(V, 2, 3, 0, 1);
-			VectorRegister4Float Sum  = VectorAdd(V, Shuf);
-			Shuf = VectorSwizzle(Sum, 1, 0, 3, 2);
-			Sum  = VectorAdd(Sum, Shuf);
-			return VectorGetComponent(Sum, 0);
-		};
-
-		float Dot  = HSum(VecDot);
-		float NormA = HSum(VecNormA);
-		float NormB = HSum(VecNormB);
-
-		// Scalar tail for dimensions not divisible by 4.
-		for (int32 i = NumVec; i < Num; ++i)
-		{
-			Dot   += DataA[i] * DataB[i];
-			NormA += DataA[i] * DataA[i];
-			NormB += DataB[i] * DataB[i];
-		}
-
-		const float Denom = FMath::Sqrt(NormA) * FMath::Sqrt(NormB);
-		return Denom > UE_SMALL_NUMBER ? Dot / Denom : 0.f;
-	}
-
 	// ---------------------------------------------------------------------------
 
 	class FSemanticSearchAssetViewOverride : public ICBSearchOverrideInstance, public TSharedFromThis<FSemanticSearchAssetViewOverride>
@@ -204,54 +149,41 @@ namespace Private
 				return;
 			}
 
-			// Shared embedding avoids copying the query vector once per item callback.
-			TSharedRef<const TArray<float>> SharedQuery = MakeShared<const TArray<float>>(MoveTemp(QueryEmbedding));
-
-			using FScoreEntry = TPair<TSharedPtr<FAssetViewItem>, float>;
-			TSharedRef<TArray<FScoreEntry>> Scores = MakeShared<TArray<FScoreEntry>>();
-			TSharedRef<int32> Countdown = MakeShared<int32>(CurrentItems.Num());
-
-			TWeakPtr<FSemanticSearchAssetViewOverride> WeakSelf = AsShared();
+			// Build ID filter and ID→Item map from current visible items
+			TArray<int64> IDFilter;
+			TMap<int64, TSharedPtr<FAssetViewItem>> IDToItem;
+			IDFilter.Reserve(CurrentItems.Num());
 
 			for (const TSharedPtr<FAssetViewItem>& Item : CurrentItems)
 			{
 				FAssetData AssetData;
-				if (!Item->GetItem().Legacy_TryGetAssetData(AssetData))
+				if (Item->GetItem().Legacy_TryGetAssetData(AssetData))
 				{
-					if (--(*Countdown) == 0)
-					{
-						PublishScores(MoveTemp(*Scores));
-					}
-					continue;
+					const int64 ID = GetAssetIndexID(AssetData);
+					IDFilter.Add(ID);
+					IDToItem.Add(ID, Item);
 				}
+			}
 
-				FAssetProcessorManager::Get().GetEmbeddingData(AssetData,
-					[WeakSelf, Epoch, Countdown, Scores, SharedQuery, Item]
-					(FAssetEmbeddingResult&& Result, FString&& Error) mutable
+			// Single hybrid search call — vector + BM25 fusion, scoped to visible items
+			const int32 K = IDFilter.Num();
+			TArray<FHybridSearchResult> Results =
+				FHybridSearchIndex::Get().Search(
+					ActiveQuery.ToString(), QueryEmbedding, K, IDFilter);
+
+			// Map results back to FAssetViewItems
+			using FScoreEntry = TPair<TSharedPtr<FAssetViewItem>, float>;
+			TArray<FScoreEntry> Scores;
+			Scores.Reserve(Results.Num());
+			for (const FHybridSearchResult& Result : Results)
+			{
+				if (TSharedPtr<FAssetViewItem>* ItemPtr = IDToItem.Find(Result.ID))
 				{
-					AsyncTask(ENamedThreads::GameThread,
-						[WeakSelf, Epoch, Countdown, Scores, SharedQuery, Item,
-						 Result = MoveTemp(Result), Error = MoveTemp(Error)]() mutable
-					{
-						TSharedPtr<FSemanticSearchAssetViewOverride> Self = WeakSelf.Pin();
-						if (!Self || Self->ActiveEpoch != Epoch)
-						{
-							return;
-						}
-
-						if (Error.IsEmpty() && !Result.Embedding.IsEmpty())
-						{
-							const float Similarity = CosineSimilarity(*SharedQuery, Result.Embedding);
-							Scores->Add({Item, Similarity});
-						}
-
-						if (--(*Countdown) == 0)
-						{
-							Self->PublishScores(MoveTemp(*Scores));
-						}
-					});
-				});
+					Scores.Add({*ItemPtr, Result.RRFScore});
+				}
 			}
+
+			PublishScores(MoveTemp(Scores));
 		}
 
 		void PublishScores(TArray<TPair<TSharedPtr<FAssetViewItem>, float>>&& Scores)
