Fix assetprocessormanager to look at parent classes for an asset to check if it is supported.
Fix crash in thumbnailprocessor

#Synced-CL 52357265

diff --git a/Engine/Plugins/Experimental/SemanticSearch/Source/SemanticSearch/Private/AssetProcessorManager.cpp b/Engine/Plugins/Experimental/SemanticSearch/Source/SemanticSearch/Private/AssetProcessorManager.cpp
--- a/Engine/Plugins/Experimental/SemanticSearch/Source/SemanticSearch/Private/AssetProcessorManager.cpp
+++ b/Engine/Plugins/Experimental/SemanticSearch/Source/SemanticSearch/Private/AssetProcessorManager.cpp
@@ -43,6 +43,28 @@ namespace UE::SemanticSearch
 			}
 		}
 
+		// Walk parent classes if no exact match
+		if (!AssetProcessor)
+		{
+			UClass* AssetClass = FindObject<UClass>(InAssetData.AssetClassPath);
+			if (AssetClass)
+			{
+				FWriteScopeLock WriteLock(AssetTypeToProcessorLock);
+				for (UClass* Super = AssetClass->GetSuperClass(); Super; Super = Super->GetSuperClass())
+				{
+					if (TSharedPtr<IAssetProcessor>* PtrToProcessor = AssetTypeToProcessor.Find(Super->GetClassPathName()))
+					{
+						if ((*PtrToProcessor)->SupportDerivedClasses())
+						{
+							AssetProcessor = *PtrToProcessor;
+							AssetTypeToProcessor.Add(InAssetData.AssetClassPath, AssetProcessor);
+							break;
+						}
+					}
+				}
+			}
+		}
+
 		if (!AssetProcessor)
 		{
 			return false;
diff --git a/Engine/Plugins/Experimental/SemanticSearch/Source/SemanticSearch/Private/Implementations/ThumbnailBaseProcessor.cpp b/Engine/Plugins/Experimental/SemanticSearch/Source/SemanticSearch/Private/Implementations/ThumbnailBaseProcessor.cpp
--- a/Engine/Plugins/Experimental/SemanticSearch/Source/SemanticSearch/Private/Implementations/ThumbnailBaseProcessor.cpp
+++ b/Engine/Plugins/Experimental/SemanticSearch/Source/SemanticSearch/Private/Implementations/ThumbnailBaseProcessor.cpp
@@ -78,8 +78,15 @@ void FThumbnailBaseAssetProcessor::GenerateCaptionRequest(const TSharedRef<const
 				Thumbnail->CompressImageData();
 				Thumbnail->AccessImageData().Empty();
 
+				FThumbnailCompressionInterface* Compressor = Thumbnail->GetCompressor();
+				if (!Compressor)
+				{
+					OnRequestComplete(false, FCaptionRequest(), FString(TEXT("Thumbnail compressor is null")));
+					return;
+				}
+
 				FAssetMedia Media;
-				Media.MimeType = Thumbnail->GetCompressor()->GetMimeType();
+				Media.MimeType = Compressor->GetMimeType();
 				Media.Data = MoveTemp(Thumbnail->AccessCompressedImageData());
 				CaptionRequest.AssetMedia.Add(MoveTemp(Media));
 
diff --git a/Engine/Plugins/Experimental/SemanticSearch/Source/SemanticSearch/Public/AssetProcessorManager.h b/Engine/Plugins/Experimental/SemanticSearch/Source/SemanticSearch/Public/AssetProcessorManager.h
--- a/Engine/Plugins/Experimental/SemanticSearch/Source/SemanticSearch/Public/AssetProcessorManager.h
+++ b/Engine/Plugins/Experimental/SemanticSearch/Source/SemanticSearch/Public/AssetProcessorManager.h
@@ -88,7 +88,7 @@ void FAssetProcessorManager::UnregisterAssetProcessor()
 {
 	static_assert(TIsDerivedFrom<FAssetProcessor, IAssetProcessor>::Value, "Asset Processors must derive from IAssetProcessor");
 	FAssetProcessor AssetProcessor;
-	UnregisterAssetProcessor(AssetProcessor.SupportDerivedClasses(), AssetProcessor.GetSupportedClass());
+	UnregisterAssetProcessor(AssetProcessor.GetSupportedClass());
 }
 
 }
\ No newline at end of file
