Fix memory leak warning

#Synced-CL 52357265

diff --git a/Engine/Plugins/Experimental/SemanticSearch/Source/SemanticSearch/Private/RemoteEmbeddingProvider.cpp b/Engine/Plugins/Experimental/SemanticSearch/Source/SemanticSearch/Private/RemoteEmbeddingProvider.cpp
--- a/Engine/Plugins/Experimental/SemanticSearch/Source/SemanticSearch/Private/RemoteEmbeddingProvider.cpp
+++ b/Engine/Plugins/Experimental/SemanticSearch/Source/SemanticSearch/Private/RemoteEmbeddingProvider.cpp
@@ -285,7 +285,7 @@ void FRemoteEmbeddingProvider::GenerateCaptionAsync(
 	if (RequestOwner)
 	{
 		// Create DDC request wrapper immediately so Begin() runs in the caller's barrier context
-		auto* DdcRequest = new FSemanticHttpRequest(Http, *RequestOwner);
+		TRefCountPtr<FSemanticHttpRequest> DdcRequest = new FSemanticHttpRequest(Http, *RequestOwner);
 		CaptionQueue.Enqueue([DdcRequest]() { DdcRequest->Start(); });
 	}
 	else
@@ -396,7 +396,7 @@ void FRemoteEmbeddingProvider::GenerateEmbeddingAsync(
 
 	if (RequestOwner)
 	{
-		auto* DdcRequest = new FSemanticHttpRequest(Http, *RequestOwner);
+		TRefCountPtr<FSemanticHttpRequest> DdcRequest = new FSemanticHttpRequest(Http, *RequestOwner);
 		EmbedQueue.Enqueue([DdcRequest]() { DdcRequest->Start(); });
 	}
 	else
