Add a progress bar for indexing

#Synced-CL 52357265

diff --git a/Engine/Plugins/Experimental/SemanticSearch/Source/SemanticSearch/Private/SemanticSearchModule.cpp b/Engine/Plugins/Experimental/SemanticSearch/Source/SemanticSearch/Private/SemanticSearchModule.cpp
--- a/Engine/Plugins/Experimental/SemanticSearch/Source/SemanticSearch/Private/SemanticSearchModule.cpp
+++ b/Engine/Plugins/Experimental/SemanticSearch/Source/SemanticSearch/Private/SemanticSearchModule.cpp
@@ -8,7 +8,9 @@
 
 #include "Modules/ModuleManager.h"
 #include "RemoteEmbeddingProvider.h"
+#include "Async/Async.h"
 #include "AssetRegistry/AssetRegistryModule.h"
+#include "Editor.h"
 #include "HAL/ConsoleManager.h"
 
 DEFINE_LOG_CATEGORY(LogSemanticSearch);
@@ -87,6 +89,11 @@ void FSemanticSearchModule::ShutdownModule()
 {
 	UE_LOG(LogSemanticSearch, Log, TEXT("SemanticSearch module shutting down"));
 
+	if (GEditor)
+	{
+		GEditor->GetTimerManager()->ClearTimer(DismissTimerHandle);
+	}
+
 	if (FAssetRegistryModule* AssetRegistryModule = FModuleManager::GetModulePtr<FAssetRegistryModule>(TEXT("AssetRegistry")))
 	{
 		IAssetRegistry& AssetRegistry = AssetRegistryModule->Get();
@@ -99,7 +106,9 @@ void FSemanticSearchModule::ShutdownModule()
 
 void FSemanticSearchModule::TriggerIndexing(const FAssetData& AssetData, bool bBuildOnMiss)
 {
-	FAssetProcessorManager::Get().GetEmbeddingData(AssetData,
+	OnIndexingStarted();
+
+	bool bStarted = FAssetProcessorManager::Get().GetEmbeddingData(AssetData,
 		[Asset = FAssetData(AssetData), bBuildOnMiss](FAssetEmbeddingResult&& EmbeddingResult, FString&& Error) mutable
 		{
 			if (!Error.IsEmpty() || EmbeddingResult.Embedding.IsEmpty())
@@ -110,6 +119,7 @@ void FSemanticSearchModule::TriggerIndexing(const FAssetData& AssetData, bool bB
 						*Asset.GetObjectPathString(),
 						Error.IsEmpty() ? TEXT("empty embedding") : *Error);
 				}
+				FSemanticSearchModule::Get().OnIndexingFinished(false);
 				return;
 			}
 
@@ -122,8 +132,23 @@ void FSemanticSearchModule::TriggerIndexing(const FAssetData& AssetData, bool bB
 					UE_LOG(LogSemanticSearch, Log, TEXT("Indexed %s (%d dims, index size: %lld)"),
 						*Asset.GetObjectPathString(), Embedding.Num(),
 						Index.GetVectorIndex() ? Index.GetVectorIndex()->GetCount() : 0);
+					FSemanticSearchModule::Get().OnIndexingFinished(true);
 				});
 		}, bBuildOnMiss);
+
+	if (!bStarted)
+	{
+		// Unsupported asset type
+		--IndexingTotal;
+		UpdateProgressNotification();
+
+		if (IndexingCompleted + IndexingFailed >= IndexingTotal && GEditor)
+		{
+			GEditor->GetTimerManager()->SetTimer(DismissTimerHandle,
+				FTimerDelegate::CreateRaw(this, &FSemanticSearchModule::DismissProgressNotification),
+				5.0f, /*bLoop=*/false);
+		}
+	}
 }
 
 void FSemanticSearchModule::OnAssetAdded(const FAssetData& AssetData)
@@ -214,4 +239,84 @@ void FSemanticSearchModule::IndexAsset(const FAssetData& AssetData)
 	TriggerIndexing(AssetData, /*bBuildOnMiss=*/true);
 }
 
+void FSemanticSearchModule::OnIndexingStarted()
+{
+	++IndexingTotal;
+
+	if (!ProgressHandle.IsValid())
+	{
+		ProgressHandle = FSlateNotificationManager::Get().StartProgressNotification(
+			FText::FromString(FString::Printf(TEXT("Indexed 0/%d"), IndexingTotal)), IndexingTotal);
+	}
+	else
+	{
+		UpdateProgressNotification();
+	}
+
+	if (DismissTimerHandle.IsValid() && GEditor)
+	{
+		GEditor->GetTimerManager()->ClearTimer(DismissTimerHandle);
+	}
+}
+
+void FSemanticSearchModule::OnIndexingFinished(bool bSuccess)
+{
+	AsyncTask(ENamedThreads::GameThread, [this, bSuccess]()
+	{
+		if (bSuccess)
+		{
+			++IndexingCompleted;
+		}
+		else
+		{
+			++IndexingFailed;
+		}
+
+		UpdateProgressNotification();
+
+		if (IndexingCompleted + IndexingFailed >= IndexingTotal && GEditor)
+		{
+			GEditor->GetTimerManager()->SetTimer(DismissTimerHandle,
+				FTimerDelegate::CreateRaw(this, &FSemanticSearchModule::DismissProgressNotification),
+				5.0f, /*bLoop=*/false);
+		}
+	});
+}
+
+void FSemanticSearchModule::UpdateProgressNotification()
+{
+	if (!ProgressHandle.IsValid())
+	{
+		return;
+	}
+
+	FString Text;
+	if (IndexingFailed > 0)
+	{
+		Text = FString::Printf(TEXT("Indexed %d/%d | Failed %d"),
+			IndexingCompleted, IndexingTotal, IndexingFailed);
+	}
+	else
+	{
+		Text = FString::Printf(TEXT("Indexed %d/%d"),
+			IndexingCompleted, IndexingTotal);
+	}
+
+	FSlateNotificationManager::Get().UpdateProgressNotification(
+		ProgressHandle, IndexingCompleted + IndexingFailed, IndexingTotal,
+		FText::FromString(Text));
+}
+
+void FSemanticSearchModule::DismissProgressNotification()
+{
+	if (ProgressHandle.IsValid())
+	{
+		FSlateNotificationManager::Get().CancelProgressNotification(ProgressHandle);
+		ProgressHandle = FProgressNotificationHandle();
+	}
+	IndexingTotal = 0;
+	IndexingCompleted = 0;
+	IndexingFailed = 0;
+}
+
 }
diff --git a/Engine/Plugins/Experimental/SemanticSearch/Source/SemanticSearch/Private/SemanticSearchModule.h b/Engine/Plugins/Experimental/SemanticSearch/Source/SemanticSearch/Private/SemanticSearchModule.h
--- a/Engine/Plugins/Experimental/SemanticSearch/Source/SemanticSearch/Private/SemanticSearchModule.h
+++ b/Engine/Plugins/Experimental/SemanticSearch/Source/SemanticSearch/Private/SemanticSearchModule.h
@@ -5,9 +5,11 @@
 #include "ISemanticSearchModule.h"
 
 #include "CoreMinimal.h"
+#include "Framework/Notifications/NotificationManager.h"
 #include "Modules/ModuleInterface.h"
 #include "Templates/SharedPointer.h"
 #include "AssetRegistry/AssetData.h"
+#include "TimerManager.h"
 
 DECLARE_LOG_CATEGORY_EXTERN(LogSemanticSearch, Log, All);
 
@@ -35,10 +37,21 @@ private:
 	void OnAssetRenamed(const FAssetData& NewAssetData, const FString& OldObjectPath);
 	void OnAssetRemoved(const FAssetData& AssetData);
 
+	void OnIndexingStarted();
+	void OnIndexingFinished(bool bSuccess);
+	void UpdateProgressNotification();
+	void DismissProgressNotification();
+
 	TSharedPtr<IEmbeddingProvider> RegisteredEmbeddingProvider;
 	FDelegateHandle AssetAddedDelegateHandle;
 	FDelegateHandle AssetUpdatedOnDiskDelegateHandle;
 	FDelegateHandle AssetRenamedDelegateHandle;
 	FDelegateHandle AssetRemovedDelegateHandle;
+
+	FProgressNotificationHandle ProgressHandle;
+	int32 IndexingTotal = 0;
+	int32 IndexingCompleted = 0;
+	int32 IndexingFailed = 0;
+	FTimerHandle DismissTimerHandle;
 };
 }
diff --git a/Engine/Plugins/Experimental/SemanticSearch/Source/SemanticSearch/SemanticSearch.Build.cs b/Engine/Plugins/Experimental/SemanticSearch/Source/SemanticSearch/SemanticSearch.Build.cs
--- a/Engine/Plugins/Experimental/SemanticSearch/Source/SemanticSearch/SemanticSearch.Build.cs
+++ b/Engine/Plugins/Experimental/SemanticSearch/Source/SemanticSearch/SemanticSearch.Build.cs
@@ -23,6 +23,7 @@ public class SemanticSearch : ModuleRules
 				"ImageCore",
 				"Json",
 				"PhysicsCore",
+				"Slate",
 				"UnrealEd",
 				"DeveloperSettings",
 				"FAISS"
