Add index stats NavBar widget to content browser

Displays live indexed/unindexed/failed asset counts in the content browser navigation bar when AI search is active. Stats update per-frame with no blocking ? indexed count reads FAISS ntotal directly, supported count is cached and invalidated on asset registry changest.Also includes move indexing CVars (IndexedPaths, AutoIndexOnStartup, AutoIndexOnSave) to USemanticSearchSettings developer settings

#Synced-CL 52378035

diff --git a/Engine/Plugins/Experimental/SemanticSearch/Source/SemanticSearch/Private/AssetProcessorManager.cpp b/Engine/Plugins/Experimental/SemanticSearch/Source/SemanticSearch/Private/AssetProcessorManager.cpp
--- a/Engine/Plugins/Experimental/SemanticSearch/Source/SemanticSearch/Private/AssetProcessorManager.cpp
+++ b/Engine/Plugins/Experimental/SemanticSearch/Source/SemanticSearch/Private/AssetProcessorManager.cpp
@@ -2,6 +2,8 @@
 
 #include "AssetProcessorManager.h"
 
+#include "AssetRegistry/AssetRegistryModule.h"
+#include "ISemanticSearchModule.h"
 #include "Misc/ScopeRWLock.h"
 #include "SemanticSearchDerivedDataBuilder.h"
 #include "UObject/Class.h"
@@ -22,6 +24,13 @@ namespace UE::SemanticSearch
 		OnObjectsReinstancedHandle = FCoreUObjectDelegates::OnObjectsReinstanced.AddRaw(this, &FAssetProcessorManager::OnObjectsReinstanced);
 		OnModuleUnloadedHandle = FCoreUObjectDelegates::CompiledInUObjectsRemovedDelegate.AddRaw(this, &FAssetProcessorManager::OnModuleUnloaded);
 		OnHotReloadCompleteHandle = FCoreUObjectDelegates::ReloadReinstancingCompleteDelegate.AddRaw(this, &FAssetProcessorManager::OnHotReloadComplete);
+
+		if (FAssetRegistryModule* ARM = FModuleManager::GetModulePtr<FAssetRegistryModule>(TEXT("AssetRegistry")))
+		{
+			IAssetRegistry& AR = ARM->Get();
+			OnAssetAddedHandle = AR.OnAssetAdded().AddRaw(this, &FAssetProcessorManager::InvalidateSupportedAssetCount);
+			OnAssetRemovedHandle = AR.OnAssetRemoved().AddRaw(this, &FAssetProcessorManager::InvalidateSupportedAssetCount);
+		}
 	}
 
 	FAssetProcessorManager::~FAssetProcessorManager()
@@ -29,42 +38,75 @@ namespace UE::SemanticSearch
 		FCoreUObjectDelegates::OnObjectsReinstanced.Remove(OnObjectsReinstancedHandle);
 		FCoreUObjectDelegates::CompiledInUObjectsRemovedDelegate.Remove(OnModuleUnloadedHandle);
 		FCoreUObjectDelegates::ReloadReinstancingCompleteDelegate.Remove(OnHotReloadCompleteHandle);
+
+		if (FAssetRegistryModule* ARM = FModuleManager::GetModulePtr<FAssetRegistryModule>(TEXT("AssetRegistry")))
+		{
+			IAssetRegistry& AR = ARM->Get();
+			AR.OnAssetAdded().Remove(OnAssetAddedHandle);
+			AR.OnAssetRemoved().Remove(OnAssetRemovedHandle);
+		}
 	}
 
-	bool FAssetProcessorManager::GetOrCreateTask(const FAssetData& InAssetData, TFunctionRef<void(const TSharedRef<Private::FSemanticSearchBuildCacheTask>&)> InApply)
+	TSharedPtr<IAssetProcessor> FAssetProcessorManager::GetProcessorForAsset(const FAssetData& InAssetData) const
+	{
+		FReadScopeLock ReadLock(AssetTypeToProcessorLock);
+		if (const TSharedPtr<IAssetProcessor>* Ptr = AssetTypeToProcessor.Find(InAssetData.AssetClassPath))
+		{
+			return *Ptr;
+		}
+		return nullptr;
+	}
+
+	int32 FAssetProcessorManager::GetSupportedAssetCount() const
+	{
+		if (CachedSupportedAssetCount < 0)
+		{
+			CachedSupportedAssetCount = ComputeSupportedAssetCount();
+		}
+		return CachedSupportedAssetCount;
+	}
+
+	void FAssetProcessorManager::InvalidateSupportedAssetCount(const FAssetData&)
+	{
+		CachedSupportedAssetCount = -1;
+	}
+
+	int32 FAssetProcessorManager::ComputeSupportedAssetCount() const
 	{
-		TSharedPtr<IAssetProcessor> AssetProcessor;
+		// Collect all unique class paths that have a processor entry.
+		// Derived classes are pre-populated by PopulateChildClassesCache.
+		TSet<FTopLevelAssetPath> SupportedClassPaths;
 		{
 			FReadScopeLock ReadLock(AssetTypeToProcessorLock);
-			TSharedPtr<IAssetProcessor>* PtrToAssetProcessor = AssetTypeToProcessor.Find(InAssetData.AssetClassPath);
-			if (PtrToAssetProcessor)
+			for (const auto& Pair : AssetTypeToProcessor)
 			{
-				AssetProcessor = *PtrToAssetProcessor;
+				if (Pair.Value)
+				{
+					SupportedClassPaths.Add(Pair.Key);
+				}
 			}
 		}
 
-		// Walk parent classes if no exact match
-		if (!AssetProcessor)
+		int32 Total = 0;
+		IAssetRegistry& AR = IAssetRegistry::GetChecked();
+		for (const FTopLevelAssetPath& ClassPath : SupportedClassPaths)
 		{
-			UClass* AssetClass = FindObject<UClass>(InAssetData.AssetClassPath);
-			if (AssetClass)
+			TArray<FAssetData> Assets;
+			AR.GetAssetsByClass(ClassPath, Assets, /*bSearchSubClasses=*/false);
+			for (const FAssetData& Asset : Assets)
 			{
-				FWriteScopeLock WriteLock(AssetTypeToProcessorLock);
-				for (UClass* Super = AssetClass->GetSuperClass(); Super; Super = Super->GetSuperClass())
+				if (!Asset.IsRedirector() && ISemanticSearchModule::IsInIndexedFolder(Asset))
 				{
-					if (TSharedPtr<IAssetProcessor>* PtrToProcessor = AssetTypeToProcessor.Find(Super->GetClassPathName()))
-					{
-						if ((*PtrToProcessor)->SupportDerivedClasses())
-						{
-							AssetProcessor = *PtrToProcessor;
-							AssetTypeToProcessor.Add(InAssetData.AssetClassPath, AssetProcessor);
-							break;
-						}
-					}
+					++Total;
 				}
 			}
 		}
+		return Total;
+	}
 
+	bool FAssetProcessorManager::GetOrCreateTask(const FAssetData& InAssetData, TFunctionRef<void(const TSharedRef<Private::FSemanticSearchBuildCacheTask>&)> InApply)
+	{
+		TSharedPtr<IAssetProcessor> AssetProcessor = GetProcessorForAsset(InAssetData);
 		if (!AssetProcessor)
 		{
 			return false;
@@ -144,13 +186,13 @@ namespace UE::SemanticSearch
 			UnregisterAssetProcessor(AssetClass);
 		}
 
+		TSharedPtr<IAssetProcessor> ProcessorPtr = AssetProcessor;
 		{
 			FWriteScopeLock WriteLock(AssetTypeToProcessorLock);
 			AssetTypeToProcessor.Add(ClassPath, MoveTemp(AssetProcessor));
 		}
 
-		InvalidAssetChildClassesCache(AssetClass);
-
+		PopulateChildClassesCache(AssetClass, ProcessorPtr);
 	}
 
 	void FAssetProcessorManager::UnregisterAssetProcessor(UClass& Class)
@@ -160,23 +202,28 @@ namespace UE::SemanticSearch
 			AssetTypeToProcessor.Remove(Class.GetClassPathName());
 		}
 
-		InvalidAssetChildClassesCache(Class);
+		PopulateChildClassesCache(Class, nullptr);
 	}
 
-	void FAssetProcessorManager::InvalidAssetChildClassesCache(UClass& Class)
+	void FAssetProcessorManager::PopulateChildClassesCache(UClass& Class, TSharedPtr<IAssetProcessor> Processor)
 	{
-		// If slow due to non native maybe consider testing the handling the recursion by hand to stop early
 		constexpr bool bRecursive = true;
 		TArray<UClass*> DerivedClasses;
 		GetDerivedClasses(&Class, DerivedClasses, bRecursive);
 
-
 		FWriteScopeLock WriteLock(AssetTypeToProcessorLock);
 		for (UClass* ChildClass : DerivedClasses)
 		{
 			if (ChildClass && ChildClass->HasAnyClassFlags(CLASS_Native))
 			{
-				AssetTypeToProcessor.Remove(ChildClass->GetClassPathName());
+				if (Processor && Processor->SupportDerivedClasses())
+				{
+					AssetTypeToProcessor.Add(ChildClass->GetClassPathName(), Processor);
+				}
+				else
+				{
+					AssetTypeToProcessor.Remove(ChildClass->GetClassPathName());
+				}
 			}
 		}
 	}
diff --git a/Engine/Plugins/Experimental/SemanticSearch/Source/SemanticSearch/Private/HybridSearchIndex.cpp b/Engine/Plugins/Experimental/SemanticSearch/Source/SemanticSearch/Private/HybridSearchIndex.cpp
--- a/Engine/Plugins/Experimental/SemanticSearch/Source/SemanticSearch/Private/HybridSearchIndex.cpp
+++ b/Engine/Plugins/Experimental/SemanticSearch/Source/SemanticSearch/Private/HybridSearchIndex.cpp
@@ -69,6 +69,23 @@ FBM25Index& FHybridSearchIndex::GetBM25Index()
 	return *BM25Index;
 }
 
+int32 FHybridSearchIndex::GetIndexedCount() const
+{
+	return VectorIndex ? static_cast<int32>(VectorIndex->GetCount()) : 0;
+}
+
+void FHybridSearchIndex::MarkFailed(int64 AssetID)
+{
+	FScopeLock Lock(&FailedLock);
+	FailedAssets.Add(AssetID);
+}
+
+int32 FHybridSearchIndex::GetFailedCount() const
+{
+	FScopeLock Lock(&FailedLock);
+	return FailedAssets.Num();
+}
+
 int64 GetAssetIndexID(FStringView AssetPath)
 {
 	FBlake3 Hasher;
@@ -118,6 +135,12 @@ void FHybridSearchIndex::Add(
 	{
 		BM25Index->Add(ID, Asset.GetObjectPathString(), Caption, Keywords);
 	}
+
+	// Clear any previous failure now that this asset indexed successfully
+	{
+		FScopeLock Lock(&FailedLock);
+		FailedAssets.Remove(ID);
+	}
 }
 
 void FHybridSearchIndex::Remove(const FAssetData& Asset)
diff --git a/Engine/Plugins/Experimental/SemanticSearch/Source/SemanticSearch/Private/SemanticSearchModule.cpp b/Engine/Plugins/Experimental/SemanticSearch/Source/SemanticSearch/Private/SemanticSearchModule.cpp
--- a/Engine/Plugins/Experimental/SemanticSearch/Source/SemanticSearch/Private/SemanticSearchModule.cpp
+++ b/Engine/Plugins/Experimental/SemanticSearch/Source/SemanticSearch/Private/SemanticSearchModule.cpp
@@ -11,45 +11,18 @@
 #include "Async/Async.h"
 #include "AssetRegistry/AssetRegistryModule.h"
 #include "Editor.h"
-#include "HAL/ConsoleManager.h"
+#include "Settings/SemanticSearchSettings.h"
 
 DEFINE_LOG_CATEGORY(LogSemanticSearch);
 
-namespace UE::SemanticSearch::Private
-{
-IMPLEMENT_MODULE(FSemanticSearchModule, SemanticSearch)
-
-static TAutoConsoleVariable<bool> CVarAutoIndexOnStartup(
-	TEXT("SemanticSearch.AutoIndexOnStartup"),
-	false,
-	TEXT("When true, assets missing from the DDC cache are sent to the backend for indexing on editor startup. DDC hits always populate the FAISS index regardless."));
-
-static TAutoConsoleVariable<bool> CVarAutoIndexOnSave(
-	TEXT("SemanticSearch.AutoIndexOnSave"),
-	true,
-	TEXT("When true, assets are automatically re-indexed when saved or updated on disk."));
-
-static TAutoConsoleVariable<FString> CVarIndexedPaths(
-	TEXT("SemanticSearch.IndexedPaths"),
-	TEXT("/Game/SemanticSearchTest/"),
-	TEXT("Comma-separated list of content browser paths to include in the semantic search index. E.g. /Game/Folder1/,/Game/Folder2/"));
-
-static TArray<FString> GetIndexedFolderPaths()
-{
-	TArray<FString> Paths;
-	CVarIndexedPaths.GetValueOnGameThread().ParseIntoArray(Paths, TEXT(","), /*bCullEmpty=*/true);
-	for (FString& Path : Paths)
-	{
-		Path.TrimStartAndEndInline();
-	}
-	return Paths;
-}
-
-static bool IsPathInIndexedFolders(const FString& Path)
+bool UE::SemanticSearch::ISemanticSearchModule::IsInIndexedFolder(const FAssetData& AssetData)
 {
-	for (const FString& Folder : GetIndexedFolderPaths())
+	const USemanticSearchSettings* Settings = USemanticSearchSettings::Get();
+	if (!Settings) { return false; }
+	const FString PackagePath = AssetData.PackageName.ToString();
+	for (const FDirectoryPath& Dir : Settings->IndexedPaths)
 	{
-		if (Path.StartsWith(Folder))
+		if (PackagePath.StartsWith(Dir.Path))
 		{
 			return true;
 		}
@@ -57,17 +30,16 @@ static bool IsPathInIndexedFolders(const FString& Path)
 	return false;
 }
 
+namespace UE::SemanticSearch::Private
+{
+IMPLEMENT_MODULE(FSemanticSearchModule, SemanticSearch)
+
 FSemanticSearchModule& FSemanticSearchModule::Get()
 {
 	static const FName ModuleName(TEXT("SemanticSearch"));
 	return FModuleManager::Get().GetModuleChecked<FSemanticSearchModule>(ModuleName);
 }
 
-bool FSemanticSearchModule::IsInIndexedFolder(const FAssetData& AssetData)
-{
-	return IsPathInIndexedFolders(AssetData.PackageName.ToString());
-}
-
 void FSemanticSearchModule::StartupModule()
 {
 	UE_LOG(LogSemanticSearch, Log, TEXT("SemanticSearch module starting up"));
@@ -113,6 +85,7 @@ void FSemanticSearchModule::TriggerIndexing(const FAssetData& AssetData, bool bB
 		{
 			if (!Error.IsEmpty() || EmbeddingResult.Embedding.IsEmpty())
 			{
+				FHybridSearchIndex::Get().MarkFailed(GetAssetIndexID(Asset));
 				if (bBuildOnMiss)
 				{
 					UE_LOG(LogSemanticSearch, Warning, TEXT("Indexing failed for %s: %s"),
@@ -153,20 +126,23 @@ void FSemanticSearchModule::TriggerIndexing(const FAssetData& AssetData, bool bB
 
 void FSemanticSearchModule::OnAssetAdded(const FAssetData& AssetData)
 {
-	if (IsInIndexedFolder(AssetData))
+	if (AssetData.IsRedirector()) { return; }
+	if (ISemanticSearchModule::IsInIndexedFolder(AssetData))
 	{
+		const USemanticSearchSettings* Settings = USemanticSearchSettings::Get();
 		UE_LOG(LogSemanticSearch, Log, TEXT("OnAssetAdded: %s"), *AssetData.GetObjectPathString());
-		TriggerIndexing(AssetData, CVarAutoIndexOnStartup.GetValueOnGameThread());
+		TriggerIndexing(AssetData, Settings && Settings->bBuildUncachedOnStartup);
 	}
 }
 
 void FSemanticSearchModule::OnAssetUpdatedOnDisk(const FAssetData& AssetData)
 {
-	if (!CVarAutoIndexOnSave.GetValueOnGameThread())
+	const USemanticSearchSettings* Settings = USemanticSearchSettings::Get();
+	if (!Settings || !Settings->bAutoIndexOnSave)
 	{
 		return;
 	}
-	if (IsInIndexedFolder(AssetData))
+	if (!AssetData.IsRedirector() && ISemanticSearchModule::IsInIndexedFolder(AssetData))
 	{
 		UE_LOG(LogSemanticSearch, Log, TEXT("OnAssetUpdatedOnDisk: %s"), *AssetData.GetObjectPathString());
 		TriggerIndexing(AssetData, /*bBuildOnMiss=*/true);
@@ -175,8 +151,18 @@ void FSemanticSearchModule::OnAssetUpdatedOnDisk(const FAssetData& AssetData)
 
 void FSemanticSearchModule::OnAssetRenamed(const FAssetData& NewAssetData, const FString& OldObjectPath)
 {
+	const USemanticSearchSettings* Settings = USemanticSearchSettings::Get();
+	auto IsPathInIndexedFolders = [Settings](const FString& Path) -> bool
+	{
+		if (!Settings) { return false; }
+		for (const FDirectoryPath& Dir : Settings->IndexedPaths)
+		{
+			if (Path.StartsWith(Dir.Path)) { return true; }
+		}
+		return false;
+	};
 	const bool bOldInFolder = IsPathInIndexedFolders(OldObjectPath);
-	const bool bNewInFolder = IsInIndexedFolder(NewAssetData);
+	const bool bNewInFolder = ISemanticSearchModule::IsInIndexedFolder(NewAssetData);
 
 	UE_LOG(LogSemanticSearch, Log, TEXT("OnAssetRenamed: %s -> %s (OldInFolder=%d, NewInFolder=%d)"),
 		*OldObjectPath, *NewAssetData.GetObjectPathString(), bOldInFolder, bNewInFolder);
@@ -199,7 +185,7 @@ void FSemanticSearchModule::OnAssetRenamed(const FAssetData& NewAssetData, const
 		{
 			UE_LOG(LogSemanticSearch, Log, TEXT("OnAssetRenamed: not in index, triggering full reindex for %s"),
 				*NewAssetData.GetObjectPathString());
-			TriggerIndexing(NewAssetData, CVarAutoIndexOnSave.GetValueOnGameThread());
+			TriggerIndexing(NewAssetData, Settings && Settings->bAutoIndexOnSave);
 		}
 	}
 	else if (bOldInFolder)
@@ -212,13 +198,13 @@ void FSemanticSearchModule::OnAssetRenamed(const FAssetData& NewAssetData, const
 	{
 		UE_LOG(LogSemanticSearch, Log, TEXT("OnAssetRenamed: moved into indexed folder, indexing %s"),
 			*NewAssetData.GetObjectPathString());
-		TriggerIndexing(NewAssetData, CVarAutoIndexOnSave.GetValueOnGameThread());
+		TriggerIndexing(NewAssetData, Settings && Settings->bAutoIndexOnSave);
 	}
 }
 
 void FSemanticSearchModule::OnAssetRemoved(const FAssetData& AssetData)
 {
-	if (IsInIndexedFolder(AssetData))
+	if (ISemanticSearchModule::IsInIndexedFolder(AssetData))
 	{
 		FHybridSearchIndex::Get().Remove(AssetData);
 	}
diff --git a/Engine/Plugins/Experimental/SemanticSearch/Source/SemanticSearch/Private/SemanticSearchModule.h b/Engine/Plugins/Experimental/SemanticSearch/Source/SemanticSearch/Private/SemanticSearchModule.h
--- a/Engine/Plugins/Experimental/SemanticSearch/Source/SemanticSearch/Private/SemanticSearchModule.h
+++ b/Engine/Plugins/Experimental/SemanticSearch/Source/SemanticSearch/Private/SemanticSearchModule.h
@@ -29,7 +29,6 @@ public:
 	virtual void IndexAsset(const FAssetData& AssetData) override;
 
 private:
-	static bool IsInIndexedFolder(const FAssetData& AssetData);
 	void TriggerIndexing(const FAssetData& AssetData, bool bBuildOnMiss);
 
 	void OnAssetAdded(const FAssetData& AssetData);
diff --git a/Engine/Plugins/Experimental/SemanticSearch/Source/SemanticSearch/Private/SemanticSearchTestCommands.cpp b/Engine/Plugins/Experimental/SemanticSearch/Source/SemanticSearch/Private/SemanticSearchTestCommands.cpp
--- a/Engine/Plugins/Experimental/SemanticSearch/Source/SemanticSearch/Private/SemanticSearchTestCommands.cpp
+++ b/Engine/Plugins/Experimental/SemanticSearch/Source/SemanticSearch/Private/SemanticSearchTestCommands.cpp
@@ -10,34 +10,40 @@
 #include "Misc/FileHelper.h"
 #include "Misc/Paths.h"
 #include "AssetRegistry/AssetRegistryModule.h"
+#include "Settings/SemanticSearchSettings.h"
 
 namespace UE::SemanticSearch::Private
 {
 
 static TSharedPtr<FRemoteEmbeddingProvider> GTestProvider;
 
-// Returns all assets in the folders listed in SemanticSearch.IndexedPaths.
+// Returns all assets in the folders configured in Semantic Search settings.
 static TArray<FAssetData> GetAssetsInIndexedFolders(IAssetRegistry& AssetRegistry)
 {
-	TArray<FString> IndexedPaths;
-	if (IConsoleVariable* CVar = IConsoleManager::Get().FindConsoleVariable(TEXT("SemanticSearch.IndexedPaths")))
+	const USemanticSearchSettings* Settings = USemanticSearchSettings::Get();
+	if (!Settings || Settings->IndexedPaths.IsEmpty())
 	{
-		CVar->GetString().ParseIntoArray(IndexedPaths, TEXT(","), /*bCullEmpty=*/true);
-	}
-
-	if (IndexedPaths.IsEmpty())
-	{
-		UE_LOG(LogSemanticSearch, Warning, TEXT("GetAssetsInIndexedFolders: SemanticSearch.IndexedPaths is empty or CVar not found"));
+		UE_LOG(LogSemanticSearch, Warning, TEXT("GetAssetsInIndexedFolders: no indexed paths configured in Semantic Search settings"));
+		return {};
 	}
 
 	TArray<FAssetData> Result;
-	for (FString& Folder : IndexedPaths)
+	TSet<FSoftObjectPath> Seen;
+	for (const FDirectoryPath& Dir : Settings->IndexedPaths)
 	{
-		Folder.TrimStartAndEndInline();
-		FString FolderPath = Folder.EndsWith(TEXT("/")) ? Folder.LeftChop(1) : Folder;
-		const int32 Before = Result.Num();
-		AssetRegistry.GetAssetsByPath(FName(FolderPath), Result, /*bRecursive=*/true);
-		UE_LOG(LogSemanticSearch, Log, TEXT("GetAssetsInIndexedFolders: '%s' -> %d assets"), *FolderPath, Result.Num() - Before);
+		TArray<FAssetData> FolderAssets;
+		AssetRegistry.GetAssetsByPath(FName(*Dir.Path), FolderAssets, /*bRecursive=*/true);
+		int32 Added = 0;
+		for (FAssetData& Asset : FolderAssets)
+		{
+			if (!Seen.Contains(Asset.GetSoftObjectPath()))
+			{
+				Seen.Add(Asset.GetSoftObjectPath());
+				Result.Add(MoveTemp(Asset));
+				++Added;
+			}
+		}
+		UE_LOG(LogSemanticSearch, Log, TEXT("GetAssetsInIndexedFolders: '%s' -> %d assets (%d new)"), *Dir.Path, FolderAssets.Num(), Added);
 	}
 	return Result;
 }
diff --git a/Engine/Plugins/Experimental/SemanticSearch/Source/SemanticSearch/Public/AssetProcessorManager.h b/Engine/Plugins/Experimental/SemanticSearch/Source/SemanticSearch/Public/AssetProcessorManager.h
--- a/Engine/Plugins/Experimental/SemanticSearch/Source/SemanticSearch/Public/AssetProcessorManager.h
+++ b/Engine/Plugins/Experimental/SemanticSearch/Source/SemanticSearch/Public/AssetProcessorManager.h
@@ -44,6 +44,12 @@ public:
 	SEMANTICSEARCH_API bool GetCaptionData(const FAssetData& InAssetData, TUniqueFunction<void(FAssetCaptionResult&&, FString&&)>&& InOnDataAvailable);
 	SEMANTICSEARCH_API bool GetEmbeddingData(const FAssetData& InAssetData, TUniqueFunction<void(FAssetEmbeddingResult&&, FString&&)>&& InOnDataAvailable, bool bBuildOnMiss = true);
 
+	/** Returns the processor for the asset's class. May return null. */
+	SEMANTICSEARCH_API TSharedPtr<IAssetProcessor> GetProcessorForAsset(const FAssetData& InAssetData) const;
+
+	/** Returns the total number of supported assets in indexed folders. Cached; invalidated on asset registry changes. */
+	SEMANTICSEARCH_API int32 GetSupportedAssetCount() const;
+
 	template<typename FAssetProcessor>
 	void RegisterAssetProcessor();
 
@@ -58,7 +64,9 @@ private:
 
 	SEMANTICSEARCH_API void RegisterAssetProcessor(TSharedRef<IAssetProcessor>&& AssetProcessor);
 	SEMANTICSEARCH_API void UnregisterAssetProcessor(UClass& Class);
-	void InvalidAssetChildClassesCache(UClass& Class);
+	void PopulateChildClassesCache(UClass& Class, TSharedPtr<IAssetProcessor> Processor);
+	int32 ComputeSupportedAssetCount() const;
+	void InvalidateSupportedAssetCount(const FAssetData&);
 
 	bool GetOrCreateTask(const FAssetData& InAssetData, TFunctionRef<void(const TSharedRef<Private::FSemanticSearchBuildCacheTask>&)> InApply);
 
@@ -69,8 +77,11 @@ private:
 	FDelegateHandle OnObjectsReinstancedHandle;
 	FDelegateHandle OnModuleUnloadedHandle;
 	FDelegateHandle OnHotReloadCompleteHandle;
+	FDelegateHandle OnAssetAddedHandle;
+	FDelegateHandle OnAssetRemovedHandle;
 
-	FRWLock AssetTypeToProcessorLock;
+	mutable TAtomic<int32> CachedSupportedAssetCount{-1};
+	mutable FRWLock AssetTypeToProcessorLock;
 	TMap<FTopLevelAssetPath, TSharedPtr<IAssetProcessor>> AssetTypeToProcessor;
 
 	TStripedMap<32, DerivedData::FCacheKey, TSharedRef<Private::FSemanticSearchBuildCacheTask>> BuildsJobs;
diff --git a/Engine/Plugins/Experimental/SemanticSearch/Source/SemanticSearch/Public/HybridSearchIndex.h b/Engine/Plugins/Experimental/SemanticSearch/Source/SemanticSearch/Public/HybridSearchIndex.h
--- a/Engine/Plugins/Experimental/SemanticSearch/Source/SemanticSearch/Public/HybridSearchIndex.h
+++ b/Engine/Plugins/Experimental/SemanticSearch/Source/SemanticSearch/Public/HybridSearchIndex.h
@@ -82,9 +82,20 @@ public:
 	/** Access the BM25 index directly (e.g. for diagnostics). */
 	FBM25Index& GetBM25Index();
 
+	/** Number of assets currently in the vector index. */
+	SEMANTICSEARCH_API int32 GetIndexedCount() const;
+
+	/** Record that an asset failed to index. */
+	SEMANTICSEARCH_API void MarkFailed(int64 AssetID);
+
+	/** Number of assets that failed to index. */
+	SEMANTICSEARCH_API int32 GetFailedCount() const;
+
 private:
 	FHybridSearchIndex();
 
+	mutable FCriticalSection FailedLock;
+	TSet<int64> FailedAssets;
 	bool bInitialized = false;
 	TSharedPtr<IVectorIndex> VectorIndex;
 	TUniquePtr<FBM25Index> BM25Index;
diff --git a/Engine/Plugins/Experimental/SemanticSearch/Source/SemanticSearch/Public/ISemanticSearchModule.h b/Engine/Plugins/Experimental/SemanticSearch/Source/SemanticSearch/Public/ISemanticSearchModule.h
--- a/Engine/Plugins/Experimental/SemanticSearch/Source/SemanticSearch/Public/ISemanticSearchModule.h
+++ b/Engine/Plugins/Experimental/SemanticSearch/Source/SemanticSearch/Public/ISemanticSearchModule.h
@@ -27,5 +27,8 @@ namespace UE::SemanticSearch
 
 		/** Index a single asset, always building from the backend if not already cached. */
 		virtual void IndexAsset(const FAssetData& AssetData) = 0;
+
+		/** Returns true if the asset is in a folder configured for indexing. */
+		SEMANTICSEARCH_API static bool IsInIndexedFolder(const FAssetData& AssetData);
 	};
 }
\ No newline at end of file
diff --git a/Engine/Plugins/Experimental/SemanticSearch/Source/SemanticSearch/Public/Settings/SemanticSearchSettings.h b/Engine/Plugins/Experimental/SemanticSearch/Source/SemanticSearch/Public/Settings/SemanticSearchSettings.h
--- a/Engine/Plugins/Experimental/SemanticSearch/Source/SemanticSearch/Public/Settings/SemanticSearchSettings.h
+++ b/Engine/Plugins/Experimental/SemanticSearch/Source/SemanticSearch/Public/Settings/SemanticSearchSettings.h
@@ -27,6 +27,18 @@ public:
 	virtual FName GetCategoryName() const override { return TEXT("Plugins"); }
 	virtual FName GetSectionName() const override { return TEXT("Semantic Search"); }
 
+	UPROPERTY(EditAnywhere, Config, Category = "Indexing",
+		meta = (ToolTip = "Content browser paths to include in the semantic search index.", RelativeToGameContentDir))
+	TArray<FDirectoryPath> IndexedPaths = { {TEXT("/Game")} };
+
+	UPROPERTY(EditAnywhere, Config, Category = "Indexing",
+		meta = (ToolTip = "When true, assets without cached embeddings are sent to the backend on startup. When false, only DDC-cached assets populate the index on startup."))
+	bool bBuildUncachedOnStartup = false;
+
+	UPROPERTY(EditAnywhere, Config, Category = "Indexing",
+		meta = (ToolTip = "When true, assets are automatically re-indexed when saved or updated on disk."))
+	bool bAutoIndexOnSave = true;
+
 	UPROPERTY(EditAnywhere, Config, Category = "Vector Index")
 	ESemanticSearchIndexType IndexType = ESemanticSearchIndexType::Flat;
 
diff --git a/Engine/Plugins/Experimental/SemanticSearch/Source/SemanticSearchEditorIntegrations/Private/ContentBrowser/SemanticSearchCB.cpp b/Engine/Plugins/Experimental/SemanticSearch/Source/SemanticSearchEditorIntegrations/Private/ContentBrowser/SemanticSearchCB.cpp
--- a/Engine/Plugins/Experimental/SemanticSearch/Source/SemanticSearchEditorIntegrations/Private/ContentBrowser/SemanticSearchCB.cpp
+++ b/Engine/Plugins/Experimental/SemanticSearch/Source/SemanticSearchEditorIntegrations/Private/ContentBrowser/SemanticSearchCB.cpp
@@ -2,31 +2,39 @@
 
 #include "ContentBrowser/SemanticSearchCB.h"
 
+#include "AssetProcessorManager.h"
 #include "AssetViewTypes.h"
 #include "Async/Async.h"
 #include "Async/TaskGraphInterfaces.h"
+#include "ContentBrowserMenuContexts.h"
 #include "ContentBrowserModule.h"
 #include "Experimental/ContentBrowserSearchOverride.h"
 #include "HybridSearchIndex.h"
 #include "ISemanticSearchModule.h"
 #include "Input/Reply.h"
 #include "Interfaces/IEmbeddingProvider.h"
+#include "Misc/DelayedAutoRegister.h"
 #include "Modules/ModuleManager.h"
 #include "Styling/AppStyle.h"
 #include "Styling/SlateColor.h"
 #include "Templates/SharedPointer.h"
+#include "ToolMenus.h"
 #include "Widgets/Input/SButton.h"
+#include "Widgets/Layout/SBox.h"
 #include "Widgets/Text/STextBlock.h"
 
 DEFINE_LOG_CATEGORY_STATIC(LogSemanticSearchCB, Log, All);
 
 namespace UE::SemanticSearch::ContentBrowser
-{
-namespace Private
 {
 	using namespace UE::Editor::ContentBrowser::Extension;
 	using namespace UE::SemanticSearch;
 
+	static bool bIsAIActive = false;
+
+namespace Private
+{
+
 	// ---------------------------------------------------------------------------
 
 	class FSemanticSearchAssetViewOverride : public ICBSearchOverrideInstance, public TSharedFromThis<FSemanticSearchAssetViewOverride>
@@ -94,7 +102,6 @@ namespace Private
 			if (ActiveQuery.IsEmpty())
 			{
 				SetUserSearching(false);
-				bHasActiveResults = false;
 				return;
 			}
 			else
@@ -262,6 +269,11 @@ namespace Private
 			.OnClicked_Lambda([this]()
 			{
 				ToggleOverride();
+				bIsAIActive = IsOverrideActive();
+				if (bIsAIActive)
+				{
+					RequestQuickFrontendListRefresh();
+				}
 				return FReply::Handled();
 			})
 			[
@@ -304,3 +316,73 @@ void UnregisterContentBrowserExtension()
 }
 
 } // namespace UE::SemanticSearch::ContentBrowser
+
+// ---------------------------------------------------------------------------
+// Index stats widget in the NavigationBar, shown alongside "X items"
+// ---------------------------------------------------------------------------
+
+static FDelayedAutoRegisterHelper SemanticSearchStatsNavBarExtension(
+	EDelayedRegisterRunPhase::EndOfEngineInit,
+	[]
+	{
+		static const FName NavigationBarName("ContentBrowser.NavigationBar");
+		if (UToolMenu* NavBar = UToolMenus::Get()->ExtendMenu(NavigationBarName))
+		{
+			FToolMenuSection& StatsSection = NavBar->FindOrAddSection(
+				"SemanticSearchStats", {},
+				FToolMenuInsert("AssetViewInfo", EToolMenuInsertType::After));
+			StatsSection.Alignment = EToolMenuSectionAlign::Last;
+
+			StatsSection.AddDynamicEntry(
+				"SemanticSearchStatsEntry",
+				FNewToolMenuSectionDelegate::CreateLambda([](FToolMenuSection& InSection)
+				{
+					const UContentBrowserToolbarMenuContext* Context =
+						InSection.FindContext<UContentBrowserToolbarMenuContext>();
+					if (!Context || !Context->AssetView.IsValid())
+					{
+						return;
+					}
+
+					FToolMenuEntry& Entry = InSection.AddEntry(
+						FToolMenuEntry::InitWidget(
+							"SemanticSearchStatsEntry",
+							SNew(SBox)
+							.Padding(FMargin(4.f, 0.f, 4.f, 0.f))
+							.Visibility_Lambda([]()
+							{
+								return UE::SemanticSearch::ContentBrowser::bIsAIActive
+									? EVisibility::Visible
+									: EVisibility::Collapsed;
+							})
+							[
+								SNew(STextBlock)
+								.TextStyle(FAppStyle::Get(), "SmallText")
+								.Text_Lambda([]() -> FText
+								{
+									const UE::SemanticSearch::FHybridSearchIndex& Index = UE::SemanticSearch::FHybridSearchIndex::Get();
+									const int32 Indexed = Index.GetIndexedCount();
+									const int32 Failed = Index.GetFailedCount();
+									const int32 TotalSupported = UE::SemanticSearch::FAssetProcessorManager::Get().GetSupportedAssetCount();
+									const int32 Unindexed = FMath::Max(0, TotalSupported - Indexed - Failed);
+									if (Failed > 0)
+									{
+										return FText::Format(
+											NSLOCTEXT("SemanticSearch", "IndexStatsWithFailed",
+												"Semantic Search Index: {0} indexed | {1} unindexed | {2} failed"),
+											FText::AsNumber(Indexed),
+											FText::AsNumber(Unindexed),
+											FText::AsNumber(Failed));
+									}
+									return FText::Format(
+										NSLOCTEXT("SemanticSearch", "IndexStats",
+											"Semantic Search Index: {0} indexed | {1} unindexed"),
+										FText::AsNumber(Indexed),
+										FText::AsNumber(Unindexed));
+								})
+							],
+							FText::GetEmpty()));
+					Entry.WidgetData.StyleParams.VerticalAlignment = VAlign_Center;
+				}));
+		}
+	});
diff --git a/Engine/Plugins/Experimental/SemanticSearch/Source/SemanticSearchEditorIntegrations/SemanticSearchEditorIntegrations.Build.cs b/Engine/Plugins/Experimental/SemanticSearch/Source/SemanticSearchEditorIntegrations/SemanticSearchEditorIntegrations.Build.cs
--- a/Engine/Plugins/Experimental/SemanticSearch/Source/SemanticSearchEditorIntegrations/SemanticSearchEditorIntegrations.Build.cs
+++ b/Engine/Plugins/Experimental/SemanticSearch/Source/SemanticSearchEditorIntegrations/SemanticSearchEditorIntegrations.Build.cs
@@ -21,6 +21,7 @@ public class SemanticSearchEditorIntegrations : ModuleRules
 				"SemanticSearch",
 				"Slate",
 				"SlateCore",
+				"ToolMenus",
 				"UnrealEd",
 			});
 
