Add support to save/load index to/from a local file.
- Add a serialization helper in FBM25index
- Add a serialization method to FHybridSearchIndex, serialized the vector index, bm25 index and the package hashes for all ids in index
- Read from index file at startup, if any index embedding stale updates it
- Save index periodically and at engine exit

#Synced-CL 52357265

diff --git a/Engine/Plugins/Experimental/SemanticSearch/Source/SemanticSearch/Private/HybridSearchIndex.cpp b/Engine/Plugins/Experimental/SemanticSearch/Source/SemanticSearch/Private/HybridSearchIndex.cpp
--- a/Engine/Plugins/Experimental/SemanticSearch/Source/SemanticSearch/Private/HybridSearchIndex.cpp
+++ b/Engine/Plugins/Experimental/SemanticSearch/Source/SemanticSearch/Private/HybridSearchIndex.cpp
@@ -2,18 +2,28 @@
 
 #include "HybridSearchIndex.h"
 
+#include "SemanticSearchModule.h"
 #include "AssetRegistry/AssetData.h"
 #include "Hash/Blake3.h"
+#include "Misc/FileHelper.h"
+#include "Misc/Paths.h"
+#include "Misc/ScopeRWLock.h"
+#include "Serialization/MemoryReader.h"
+#include "Serialization/MemoryWriter.h"
 #include "Settings/SemanticSearchSettings.h"
 #include "TextIndex/BM25Index.h"
 #include "VectorIndex/FlatVectorIndex.h"
 #include "VectorIndex/PQVectorIndex.h"
 
-#include "Misc/ScopeRWLock.h"
-
 namespace UE::SemanticSearch
 {
 
+struct FHybridIndexHeader
+{
+	static constexpr uint32 Magic = 0x48594258;   // "HYBX"
+	static constexpr uint32 Version = 1;
+};
+
 FHybridSearchIndex& FHybridSearchIndex::Get()
 {
 	static FHybridSearchIndex Instance;
@@ -150,13 +160,16 @@ void FHybridSearchIndex::Remove(const FAssetData& Asset)
 
 void FHybridSearchIndex::Remove(FStringView AssetPath)
 {
-	const int64 ID = GetAssetIndexID(AssetPath);
+	Remove(GetAssetIndexID(AssetPath));
+}
 
+void FHybridSearchIndex::Remove(int64 AssetID)
+{
 	if (VectorIndex)
 	{
-		VectorIndex->Remove(ID);
+		VectorIndex->Remove(AssetID);
 	}
-	BM25Index->Remove(ID);
+	BM25Index->Remove(AssetID);
 }
 
 bool FHybridSearchIndex::Contains(const FAssetData& Asset) const
@@ -292,4 +305,146 @@ TArray<FHybridSearchResult> FHybridSearchIndex::Search(
 	return Results;
 }
 
+bool FHybridSearchIndex::SaveToFile(const FString& FilePath, TMap<int64, FIoHash>& PackageHashes) const
+{
+	FReadScopeLock ReadLock(IndexLock);
+
+	TArray<uint8> FileData;
+	FMemoryWriter Ar(FileData);
+
+	uint32 WriteMagic = FHybridIndexHeader::Magic;
+	uint32 WriteVersion = FHybridIndexHeader::Version;
+	Ar << WriteMagic << WriteVersion;
+
+	bool bWriteInitialized = bInitialized;
+	Ar << bWriteInitialized;
+
+	// Vector index type + data
+	{
+		const USemanticSearchSettings* Settings = USemanticSearchSettings::Get();
+		uint8 IndexType = (bInitialized && VectorIndex && Settings)
+			? static_cast<uint8>(Settings->IndexType)
+			: static_cast<uint8>(ESemanticSearchIndexType::Flat);
+		Ar << IndexType;
+	}
+
+	if (bInitialized && VectorIndex)
+	{
+		int32 VectorDim = VectorIndex->GetDimension();
+		TArray<uint8> VectorData = VectorIndex->Serialize();
+		Ar << VectorDim;
+		Ar << VectorData;
+	}
+	else
+	{
+		int32 VectorDim = 0;
+		TArray<uint8> VectorData;
+		Ar << VectorDim;
+		Ar << VectorData;
+	}
+
+	// BM25 index
+	BM25Index->Serialize(Ar);
+
+	// Staleness manifest
+	Ar << PackageHashes;
+
+	// Ensure directory exists
+	FString Directory = FPaths::GetPath(FilePath);
+	IFileManager::Get().MakeDirectory(*Directory, /*bTree=*/true);
+
+	return FFileHelper::SaveArrayToFile(FileData, *FilePath);
+}
+
+bool FHybridSearchIndex::LoadFromFile(const FString& FilePath, TMap<int64, FIoHash>& OutPackageHashes)
+{
+	TArray<uint8> FileData;
+	if (!FFileHelper::LoadFileToArray(FileData, *FilePath))
+	{
+		return false;
+	}
+
+	FWriteScopeLock WriteLock(IndexLock);
+
+	FMemoryReader Ar(FileData);
+
+	uint32 Magic = 0;
+	uint32 Version = 0;
+	Ar << Magic << Version;
+	if (Magic != FHybridIndexHeader::Magic || Version != FHybridIndexHeader::Version)
+	{
+		return false;
+	}
+
+	// Deserialize into temporaries — only swap into members on full success
+	bool LoadedInitialized = false;
+	Ar << LoadedInitialized;
+
+	uint8 IndexType = 0;
+	Ar << IndexType;
+
+	int32 VectorDim = 0;
+	TArray<uint8> VectorData;
+	Ar << VectorDim;
+	Ar << VectorData;
+
+	if (Ar.IsError())
+	{
+		return false;
+	}
+
+	TSharedPtr<IVectorIndex> LoadedVectorIndex;
+	if (LoadedInitialized && VectorDim > 0 && VectorData.Num() > 0)
+	{
+		const USemanticSearchSettings* Settings = USemanticSearchSettings::Get();
+		const ESemanticSearchIndexType SavedType = static_cast<ESemanticSearchIndexType>(IndexType);
+		const ESemanticSearchIndexType CurrentType = Settings ? Settings->IndexType : ESemanticSearchIndexType::Flat;
+
+		if (SavedType != CurrentType)
+		{
+			UE_LOG(LogSemanticSearch, Warning, TEXT("Saved index type (%d) does not match current settings (%d), discarding"),
+				IndexType, static_cast<uint8>(CurrentType));
+			return false;
+		}
+
+		if (SavedType == ESemanticSearchIndexType::PQ)
+		{
+			LoadedVectorIndex = TSharedPtr<IVectorIndex>(FPQVectorIndex::Deserialize(VectorData, VectorDim).Release());
+		}
+		else
+		{
+			LoadedVectorIndex = TSharedPtr<IVectorIndex>(FFlatVectorIndex::Deserialize(VectorData, VectorDim).Release());
+		}
+
+		if (!LoadedVectorIndex)
+		{
+			return false;
+		}
+	}
+
+	// BM25 index
+	TUniquePtr<FBM25Index> LoadedBM25 = FBM25Index::Deserialize(Ar);
+	if (!LoadedBM25)
+	{
+		return false;
+	}
+
+	// Staleness manifest
+	TMap<int64, FIoHash> LoadedHashes;
+	Ar << LoadedHashes;
+
+	if (Ar.IsError())
+	{
+		return false;
+	}
+
+	// All deserialization succeeded — swap into members atomically
+	bInitialized = LoadedInitialized;
+	VectorIndex = MoveTemp(LoadedVectorIndex);
+	BM25Index = MoveTemp(LoadedBM25);
+	OutPackageHashes = MoveTemp(LoadedHashes);
+
+	return true;
+}
+
 } // namespace UE::SemanticSearch
diff --git a/Engine/Plugins/Experimental/SemanticSearch/Source/SemanticSearch/Private/SemanticSearchModule.cpp b/Engine/Plugins/Experimental/SemanticSearch/Source/SemanticSearch/Private/SemanticSearchModule.cpp
--- a/Engine/Plugins/Experimental/SemanticSearch/Source/SemanticSearch/Private/SemanticSearchModule.cpp
+++ b/Engine/Plugins/Experimental/SemanticSearch/Source/SemanticSearch/Private/SemanticSearchModule.cpp
@@ -5,6 +5,7 @@
 #include "AssetProcessorManager.h"
 #include "AssetProcessors/AssetProcessorUtils.h"
 #include "HybridSearchIndex.h"
+#include "TextIndex/BM25Index.h"
 
 #include "Modules/ModuleManager.h"
 #include "RemoteEmbeddingProvider.h"
@@ -12,6 +13,8 @@
 #include "AssetRegistry/AssetRegistryModule.h"
 #include "Editor.h"
 #include "Settings/SemanticSearchSettings.h"
+#include "HAL/ConsoleManager.h"
+#include "Misc/Paths.h"
 
 DEFINE_LOG_CATEGORY(LogSemanticSearch);
 
@@ -34,6 +37,18 @@ namespace UE::SemanticSearch::Private
 {
 IMPLEMENT_MODULE(FSemanticSearchModule, SemanticSearch)
 
+static FIoHash GetPackageSavedHash(const FAssetData& AssetData)
+{
+	IAssetRegistry& AssetRegistry = FModuleManager::LoadModuleChecked<FAssetRegistryModule>(TEXT("AssetRegistry")).Get();
+	TOptional<FAssetPackageData> PkgData = AssetRegistry.GetAssetPackageDataCopy(AssetData.PackageName);
+	return PkgData ? PkgData->GetPackageSavedHash() : FIoHash();
+}
+
+FString FSemanticSearchModule::GetSavedIndexPath()
+{
+	return FPaths::EngineSavedDir() / TEXT("SemanticSearch") / TEXT("SearchIndex.bin");
+}
+
 FSemanticSearchModule& FSemanticSearchModule::Get()
 {
 	static const FName ModuleName(TEXT("SemanticSearch"));
@@ -46,6 +61,17 @@ void FSemanticSearchModule::StartupModule()
 	RegisteredEmbeddingProvider = MakeShared<FRemoteEmbeddingProvider>();
 	RegisterDefaultAssetProcessors();
 
+	// Try to load saved index from disk
+	FHybridSearchIndex& Index = FHybridSearchIndex::Get();
+	const FString IndexPath = GetSavedIndexPath();
+	if (Index.LoadFromFile(IndexPath, LoadedPackageHashes))
+	{
+		int64 VectorCount = Index.GetVectorIndex() ? Index.GetVectorIndex()->GetCount() : 0;
+		int64 BM25Count = Index.GetBM25Index().GetCount();
+		UE_LOG(LogSemanticSearch, Log, TEXT("Loaded saved index from %s (vector: %lld, BM25: %lld, hashes: %d)"),
+			*IndexPath, VectorCount, BM25Count, LoadedPackageHashes.Num());
+	}
+
 	IAssetRegistry& AssetRegistry = FModuleManager::LoadModuleChecked<FAssetRegistryModule>(TEXT("AssetRegistry")).Get();
 	AssetAddedDelegateHandle = AssetRegistry.OnAssetAdded().AddRaw(
 		this, &FSemanticSearchModule::OnAssetAdded);
@@ -55,6 +81,11 @@ void FSemanticSearchModule::StartupModule()
 		this, &FSemanticSearchModule::OnAssetRenamed);
 	AssetRemovedDelegateHandle = AssetRegistry.OnAssetRemoved().AddRaw(
 		this, &FSemanticSearchModule::OnAssetRemoved);
+	FilesLoadedDelegateHandle = AssetRegistry.OnFilesLoaded().AddRaw(
+		this, &FSemanticSearchModule::OnFilesLoaded);
+
+	PreExitDelegateHandle = FCoreDelegates::OnEnginePreExit.AddRaw(
+		this, &FSemanticSearchModule::SaveIndex);
 }
 
 void FSemanticSearchModule::ShutdownModule()
@@ -64,8 +95,11 @@ void FSemanticSearchModule::ShutdownModule()
 	if (GEditor)
 	{
 		GEditor->GetTimerManager()->ClearTimer(DismissTimerHandle);
+		GEditor->GetTimerManager()->ClearTimer(PeriodicSaveTimerHandle);
 	}
 
+	FCoreDelegates::OnEnginePreExit.Remove(PreExitDelegateHandle);
+
 	if (FAssetRegistryModule* AssetRegistryModule = FModuleManager::GetModulePtr<FAssetRegistryModule>(TEXT("AssetRegistry")))
 	{
 		IAssetRegistry& AssetRegistry = AssetRegistryModule->Get();
@@ -73,6 +107,7 @@ void FSemanticSearchModule::ShutdownModule()
 		AssetRegistry.OnAssetUpdatedOnDisk().Remove(AssetUpdatedOnDiskDelegateHandle);
 		AssetRegistry.OnAssetRenamed().Remove(AssetRenamedDelegateHandle);
 		AssetRegistry.OnAssetRemoved().Remove(AssetRemovedDelegateHandle);
+		AssetRegistry.OnFilesLoaded().Remove(FilesLoadedDelegateHandle);
 	}
 }
 
@@ -126,13 +161,27 @@ void FSemanticSearchModule::TriggerIndexing(const FAssetData& AssetData, bool bB
 
 void FSemanticSearchModule::OnAssetAdded(const FAssetData& AssetData)
 {
-	if (AssetData.IsRedirector()) { return; }
-	if (ISemanticSearchModule::IsInIndexedFolder(AssetData))
+	if (AssetData.IsRedirector() || !IsInIndexedFolder(AssetData)) { return; }
+
+	FHybridSearchIndex& Index = FHybridSearchIndex::Get();
+	if (!LoadedPackageHashes.IsEmpty() && Index.Contains(AssetData))
 	{
-		const USemanticSearchSettings* Settings = USemanticSearchSettings::Get();
-		UE_LOG(LogSemanticSearch, Log, TEXT("OnAssetAdded: %s"), *AssetData.GetObjectPathString());
-		TriggerIndexing(AssetData, Settings && Settings->bBuildUncachedOnStartup);
+		// Already in index (loaded from disk) — check staleness via loaded hashes
+		const int64 ID = GetAssetIndexID(AssetData);
+		FIoHash StoredHash;
+		if (LoadedPackageHashes.RemoveAndCopyValue(ID, StoredHash))
+		{
+			FIoHash CurrentHash = GetPackageSavedHash(AssetData);
+			if (StoredHash == CurrentHash)
+			{
+				return; // Up to date, skip
+			}
+		}
 	}
+
+	const USemanticSearchSettings* Settings = USemanticSearchSettings::Get();
+	UE_LOG(LogSemanticSearch, Log, TEXT("OnAssetAdded: %s"), *AssetData.GetObjectPathString());
+	TriggerIndexing(AssetData, Settings && Settings->bBuildUncachedOnStartup);
 }
 
 void FSemanticSearchModule::OnAssetUpdatedOnDisk(const FAssetData& AssetData)
@@ -305,4 +354,68 @@ void FSemanticSearchModule::DismissProgressNotification()
 	IndexingFailed = 0;
 }
 
+void FSemanticSearchModule::OnFilesLoaded()
+{
+	// Anything left in LoadedPackageHashes was not seen during OnAssetAdded — the asset was deleted
+	if (!LoadedPackageHashes.IsEmpty())
+	{
+		FHybridSearchIndex& Index = FHybridSearchIndex::Get();
+		for (const auto& [ID, Hash] : LoadedPackageHashes)
+		{
+			Index.Remove(ID);
+		}
+		UE_LOG(LogSemanticSearch, Log, TEXT("Removed %d orphaned entries from loaded index"), LoadedPackageHashes.Num());
+		LoadedPackageHashes.Empty();
+	}
+
+	SaveIndex();
+	// Start periodic save timer
+	if (GEditor)
+	{
+		static constexpr float PeriodicSaveIntervalSeconds = 300.0f; // 5 minutes
+		GEditor->GetTimerManager()->SetTimer(PeriodicSaveTimerHandle,
+			FTimerDelegate::CreateRaw(this, &FSemanticSearchModule::SaveIndex),
+			PeriodicSaveIntervalSeconds, /*bLoop=*/true);
+	}
+}
+
+void FSemanticSearchModule::SaveIndex()
+{
+	FHybridSearchIndex& Index = FHybridSearchIndex::Get();
+	if (!Index.IsInitialized())
+	{
+		return;
+	}
+
+	// Generate package hashes for all currently indexed assets
+	TMap<int64, FIoHash> PackageHashes;
+	IAssetRegistry& AssetRegistry = FModuleManager::LoadModuleChecked<FAssetRegistryModule>(TEXT("AssetRegistry")).Get();
+	AssetRegistry.EnumerateAllAssets([&Index, &PackageHashes, &AssetRegistry](const FAssetData& AssetData)
+	{
+		if (Index.Contains(AssetData))
+		{
+			TOptional<FAssetPackageData> PkgData = AssetRegistry.GetAssetPackageDataCopy(AssetData.PackageName);
+			if (PkgData)
+			{
+				const int64 ID = GetAssetIndexID(AssetData);
+				PackageHashes.Add(ID, PkgData->GetPackageSavedHash());
+			}
+		}
+		return true;
+	});
+
+	const FString IndexPath = GetSavedIndexPath();
+	if (Index.SaveToFile(IndexPath, PackageHashes))
+	{
+		int64 VectorCount = Index.GetVectorIndex() ? Index.GetVectorIndex()->GetCount() : 0;
+		int64 BM25Count = Index.GetBM25Index().GetCount();
+		UE_LOG(LogSemanticSearch, Log, TEXT("Saved index to %s (vector: %lld, BM25: %lld, hashes: %d)"),
+			*IndexPath, VectorCount, BM25Count, PackageHashes.Num());
+	}
+	else
+	{
+		UE_LOG(LogSemanticSearch, Warning, TEXT("Failed to save index to %s"), *IndexPath);
+	}
+}
+
 }
diff --git a/Engine/Plugins/Experimental/SemanticSearch/Source/SemanticSearch/Private/SemanticSearchModule.h b/Engine/Plugins/Experimental/SemanticSearch/Source/SemanticSearch/Private/SemanticSearchModule.h
--- a/Engine/Plugins/Experimental/SemanticSearch/Source/SemanticSearch/Private/SemanticSearchModule.h
+++ b/Engine/Plugins/Experimental/SemanticSearch/Source/SemanticSearch/Private/SemanticSearchModule.h
@@ -9,6 +9,7 @@
 #include "Modules/ModuleInterface.h"
 #include "Templates/SharedPointer.h"
 #include "AssetRegistry/AssetData.h"
+#include "IO/IoHash.h"
 #include "TimerManager.h"
 
 DECLARE_LOG_CATEGORY_EXTERN(LogSemanticSearch, Log, All);
@@ -35,22 +36,32 @@ private:
 	void OnAssetUpdatedOnDisk(const FAssetData& AssetData);
 	void OnAssetRenamed(const FAssetData& NewAssetData, const FString& OldObjectPath);
 	void OnAssetRemoved(const FAssetData& AssetData);
+	void OnFilesLoaded();
 
 	void OnIndexingStarted();
 	void OnIndexingFinished(bool bSuccess);
 	void UpdateProgressNotification();
 	void DismissProgressNotification();
+	void SaveIndex();
+
+	static FString GetSavedIndexPath();
 
 	TSharedPtr<IEmbeddingProvider> RegisteredEmbeddingProvider;
 	FDelegateHandle AssetAddedDelegateHandle;
 	FDelegateHandle AssetUpdatedOnDiskDelegateHandle;
 	FDelegateHandle AssetRenamedDelegateHandle;
 	FDelegateHandle AssetRemovedDelegateHandle;
+	FDelegateHandle FilesLoadedDelegateHandle;
+	FDelegateHandle PreExitDelegateHandle;
 
 	FProgressNotificationHandle ProgressHandle;
 	int32 IndexingTotal = 0;
 	int32 IndexingCompleted = 0;
 	int32 IndexingFailed = 0;
 	FTimerHandle DismissTimerHandle;
+	FTimerHandle PeriodicSaveTimerHandle;
+
+	/** Package hashes loaded from saved index, used for staleness checks during startup. */
+	TMap<int64, FIoHash> LoadedPackageHashes;
 };
 }
diff --git a/Engine/Plugins/Experimental/SemanticSearch/Source/SemanticSearch/Private/TextIndex/BM25Index.cpp b/Engine/Plugins/Experimental/SemanticSearch/Source/SemanticSearch/Private/TextIndex/BM25Index.cpp
--- a/Engine/Plugins/Experimental/SemanticSearch/Source/SemanticSearch/Private/TextIndex/BM25Index.cpp
+++ b/Engine/Plugins/Experimental/SemanticSearch/Source/SemanticSearch/Private/TextIndex/BM25Index.cpp
@@ -9,6 +9,12 @@
 namespace UE::SemanticSearch
 {
 
+struct FBM25IndexHeader
+{
+	static constexpr uint32 Magic = 0x424D3235; // "BM25"
+	static constexpr uint32 Version = 1;
+};
+
 namespace Tokenizer
 {
 
@@ -444,4 +450,41 @@ bool FBM25Index::Update(int64 ID, FStringView OldPath, FStringView NewPath)
 	return true;
 }
 
+void FBM25Index::Serialize(FArchive& Ar)
+{
+	FReadScopeLock ReadLock(IndexLock);
+
+	uint32 WriteMagic = FBM25IndexHeader::Magic;
+	uint32 WriteVersion = FBM25IndexHeader::Version;
+	Ar << WriteMagic << WriteVersion;
+
+	Ar << PostingLists;
+	Ar << DocLengths;
+	Ar << DocTokens;
+	Ar << TotalTokenCount;
+}
+
+TUniquePtr<FBM25Index> FBM25Index::Deserialize(FArchive& Ar)
+{
+	uint32 Magic = 0;
+	uint32 Version = 0;
+	Ar << Magic << Version;
+	if (Magic != FBM25IndexHeader::Magic || Version != FBM25IndexHeader::Version)
+	{
+		return nullptr;
+	}
+
+	TUniquePtr<FBM25Index> Index = MakeUnique<FBM25Index>();
+	Ar << Index->PostingLists;
+	Ar << Index->DocLengths;
+	Ar << Index->DocTokens;
+	Ar << Index->TotalTokenCount;
+
+	if (Ar.IsError())
+	{
+		return nullptr;
+	}
+	return Index;
+}
+
 } // namespace UE::SemanticSearch
diff --git a/Engine/Plugins/Experimental/SemanticSearch/Source/SemanticSearch/Private/TextIndex/BM25Index.h b/Engine/Plugins/Experimental/SemanticSearch/Source/SemanticSearch/Private/TextIndex/BM25Index.h
--- a/Engine/Plugins/Experimental/SemanticSearch/Source/SemanticSearch/Private/TextIndex/BM25Index.h
+++ b/Engine/Plugins/Experimental/SemanticSearch/Source/SemanticSearch/Private/TextIndex/BM25Index.h
@@ -7,6 +7,7 @@
 #include "Containers/Map.h"
 #include "Containers/StringView.h"
 #include "HAL/CriticalSection.h"
+#include "Serialization/Archive.h"
 
 namespace UE::SemanticSearch
 {
@@ -44,11 +45,22 @@ public:
 	/** Update a document's path tokens without touching caption/keyword tokens. */
 	bool Update(int64 ID, FStringView OldPath, FStringView NewPath);
 
+	/** Serialize the index to an archive. */
+	void Serialize(FArchive& Ar);
+
+	/** Deserialize a BM25 index from an archive. Returns null on version mismatch or corrupt data. */
+	static TUniquePtr<FBM25Index> Deserialize(FArchive& Ar);
+
 private:
 	struct FPostingEntry
 	{
 		int64 DocID;
 		uint8 TermFreq;  // capped at 255
+
+		friend FArchive& operator<<(FArchive& Ar, FPostingEntry& Entry)
+		{
+			return Ar << Entry.DocID << Entry.TermFreq;
+		}
 	};
 
 	/** Score a single document against query tokens. */
diff --git a/Engine/Plugins/Experimental/SemanticSearch/Source/SemanticSearch/Public/HybridSearchIndex.h b/Engine/Plugins/Experimental/SemanticSearch/Source/SemanticSearch/Public/HybridSearchIndex.h
--- a/Engine/Plugins/Experimental/SemanticSearch/Source/SemanticSearch/Public/HybridSearchIndex.h
+++ b/Engine/Plugins/Experimental/SemanticSearch/Source/SemanticSearch/Public/HybridSearchIndex.h
@@ -4,8 +4,10 @@
 
 #include "Containers/Array.h"
 #include "Containers/ArrayView.h"
+#include "Containers/Map.h"
 #include "Containers/StringView.h"
 #include "Interfaces/IVectorIndex.h"
+#include "IO/IoHash.h"
 #include "Templates/SharedPointer.h"
 #include "Templates/UniquePtr.h"
 #include "HAL/CriticalSection.h"
@@ -67,6 +69,9 @@ public:
 	/** Remove by path string (for when the asset has already been deleted). */
 	SEMANTICSEARCH_API void Remove(FStringView AssetPath);
 
+	/** Remove by pre-computed ID. */
+	SEMANTICSEARCH_API void Remove(int64 AssetID);
+
 	/** Check if an asset is in the vector index. */
 	SEMANTICSEARCH_API bool Contains(const FAssetData& Asset) const;
 
@@ -91,6 +96,19 @@ public:
 	/** Number of assets that failed to index. */
 	SEMANTICSEARCH_API int32 GetFailedCount() const;
 
+	/**
+	 * Save the full index (vector + BM25 + package hashes) to a file.
+	 * @param FilePath		Destination file path.
+	 * @param PackageHashes	Per-asset PackageSavedHash for staleness detection on next load.
+	 */
+	SEMANTICSEARCH_API bool SaveToFile(const FString& FilePath, TMap<int64, FIoHash>& PackageHashes) const;
+
+	/**
+	 * Load a previously saved index from a file. Replaces current index state.
+	 * @param OutPackageHashes	Filled with per-asset PackageSavedHash from save time (for staleness checks).
+	 */
+	SEMANTICSEARCH_API bool LoadFromFile(const FString& FilePath, TMap<int64, FIoHash>& OutPackageHashes);
+
 private:
 	FHybridSearchIndex();
 
