Filter low-relevance results from AI search in the Content Browser

- SemanticSearchSettings.h ? Add VectorDistanceCutoff setting (default 1.25, range 0.1?3.0). Controls the maximum L2 distance for a result to qualify as a semantic match. Configurable in Editor Preferences > Plugins > Semantic Search.
- SemanticSearchCB.cpp ? Filter results in OnQueryEmbeddingReady(): an asset must have a BM25 keyword match (score > 0) OR a vector distance below VectorDistanceCutoff to appear. Results that fail both checks are dropped before publishing.
- SemanticSearchCB.cpp ? Fix "Applying filter..." getting stuck when all results are filtered out. When zero results pass the quality gate, PublishScores defers clearing bHasInFlightRequest so the CB's ProcessItemsPendingFilter gets one more tick to finalize and reset InitialNumAmortizedTasks. OnItemsAvailable handles the deferred finalization on the next tick.

#Synced-CL 52513209

diff --git a/Engine/Plugins/Experimental/SemanticSearch/Source/SemanticSearch/Public/Settings/SemanticSearchSettings.h b/Engine/Plugins/Experimental/SemanticSearch/Source/SemanticSearch/Public/Settings/SemanticSearchSettings.h
--- a/Engine/Plugins/Experimental/SemanticSearch/Source/SemanticSearch/Public/Settings/SemanticSearchSettings.h
+++ b/Engine/Plugins/Experimental/SemanticSearch/Source/SemanticSearch/Public/Settings/SemanticSearchSettings.h
@@ -56,4 +56,9 @@ public:
 			ClampMin = "1", ClampMax = "10",
 			ToolTip = "Oversample multiplier for RRF. Each source fetches K * this many candidates before fusion."))
 	int32 RRFOversample = 3;
+
+	UPROPERTY(EditAnywhere, Config, Category = "Search",
+		meta = (ClampMin = "0.1", ClampMax = "3.0",
+			ToolTip = "Maximum L2 distance for a result to be considered a semantic match. Lower values are stricter. Results must pass this threshold OR have a BM25 keyword match to appear."))
+	float VectorDistanceCutoff = 1.25f;
 };
diff --git a/Engine/Plugins/Experimental/SemanticSearch/Source/SemanticSearchEditorIntegrations/Private/ContentBrowser/SemanticSearchCB.cpp b/Engine/Plugins/Experimental/SemanticSearch/Source/SemanticSearchEditorIntegrations/Private/ContentBrowser/SemanticSearchCB.cpp
--- a/Engine/Plugins/Experimental/SemanticSearch/Source/SemanticSearchEditorIntegrations/Private/ContentBrowser/SemanticSearchCB.cpp
+++ b/Engine/Plugins/Experimental/SemanticSearch/Source/SemanticSearchEditorIntegrations/Private/ContentBrowser/SemanticSearchCB.cpp
@@ -11,6 +11,7 @@
 #include "Experimental/ContentBrowserSearchOverride.h"
 #include "HybridSearchIndex.h"
 #include "ISemanticSearchModule.h"
+#include "Settings/SemanticSearchSettings.h"
 #include "Input/Reply.h"
 #include "Interfaces/IEmbeddingProvider.h"
 #include "Misc/DelayedAutoRegister.h"
@@ -61,6 +62,16 @@ namespace Private
 		{
 			CurrentItems.Append(Items);
 
+			// When the search completed with zero matches, PublishScores left
+			// bHasInFlightRequest true so the CB keeps ticking.  Finalize here
+			// so the amortization counter can reset on this tick.
+			if (bHasActiveResults && bHasInFlightRequest)
+			{
+				bHasInFlightRequest = false;
+				PublishMatchingItems();
+				return;
+			}
+
 			if (Items.IsEmpty())
 			{
 				return;
@@ -183,12 +194,24 @@ namespace Private
 				FHybridSearchIndex::Get().Search(
 					ActiveQuery.ToString(), QueryEmbedding, K, IDFilter);
 
-			// Map results back to FAssetViewItems
+			// Map results back to FAssetViewItems, filtering out low-relevance noise
+			const USemanticSearchSettings* Settings = USemanticSearchSettings::Get();
+			const float DistanceCutoff = Settings ? Settings->VectorDistanceCutoff : 1.25f;
+
 			using FScoreEntry = TPair<TSharedPtr<FAssetViewItem>, float>;
 			TArray<FScoreEntry> Scores;
 			Scores.Reserve(Results.Num());
 			for (const FHybridSearchResult& Result : Results)
 			{
+				const bool bHasKeywordMatch = Result.BM25Score > 0.0f;
+				const bool bHasSemanticMatch = Result.VectorDistance >= 0.0f
+					&& Result.VectorDistance < DistanceCutoff;
+
+				if (!bHasKeywordMatch && !bHasSemanticMatch)
+				{
+					continue;
+				}
+
 				if (TSharedPtr<FAssetViewItem>* ItemPtr = IDToItem.Find(Result.ID))
 				{
 					Scores.Add({*ItemPtr, Result.RRFScore});
@@ -212,8 +235,18 @@ namespace Private
 				LastMatchedItems.Add(MoveTemp(Item));
 			}
 
-			bHasInFlightRequest = false;
 			bHasActiveResults = true;
+
+			if (LastMatchedItems.IsEmpty())
+			{
+				// Keep bHasInFlightRequest true so the CB's HasItemsPendingFilter()
+				// returns true for one more ProcessItemsPendingFilter tick — that lets
+				// the amortization counter (InitialNumAmortizedTasks) reach the reset
+				// path.  OnItemsAvailable() will finalize on the next tick.
+				return;
+			}
+
+			bHasInFlightRequest = false;
 			PublishMatchingItems();
 		}
 
