Fixed a crash on boot. The timer manager might not be valid during the early boot of the engine.
I just removed the timer if the indexing is too quick we just shouldn't use that mecanism to display the result and in those case just use a regular notification at the end of the process if that process is manually triggered by the user.

#Synced-CL 52518173

diff --git a/Engine/Plugins/Experimental/SemanticSearch/Source/SemanticSearch/Private/SemanticSearchModule.cpp b/Engine/Plugins/Experimental/SemanticSearch/Source/SemanticSearch/Private/SemanticSearchModule.cpp
--- a/Engine/Plugins/Experimental/SemanticSearch/Source/SemanticSearch/Private/SemanticSearchModule.cpp
+++ b/Engine/Plugins/Experimental/SemanticSearch/Source/SemanticSearch/Private/SemanticSearchModule.cpp
@@ -150,11 +150,9 @@ void FSemanticSearchModule::TriggerIndexing(const FAssetData& AssetData, bool bB
 		--IndexingTotal;
 		UpdateProgressNotification();
 
-		if (IndexingCompleted + IndexingFailed >= IndexingTotal && GEditor)
+		if (IndexingCompleted + IndexingFailed >= IndexingTotal)
 		{
-			GEditor->GetTimerManager()->SetTimer(DismissTimerHandle,
-				FTimerDelegate::CreateRaw(this, &FSemanticSearchModule::DismissProgressNotification),
-				5.0f, /*bLoop=*/false);
+			DismissProgressNotification();
 		}
 	}
 }
