Change semantic search threshold to 0.2

#Synced-CL 52523344

diff --git a/Engine/Plugins/Experimental/SemanticSearch/Source/SemanticSearch/Public/Settings/SemanticSearchSettings.h b/Engine/Plugins/Experimental/SemanticSearch/Source/SemanticSearch/Public/Settings/SemanticSearchSettings.h
--- a/Engine/Plugins/Experimental/SemanticSearch/Source/SemanticSearch/Public/Settings/SemanticSearchSettings.h
+++ b/Engine/Plugins/Experimental/SemanticSearch/Source/SemanticSearch/Public/Settings/SemanticSearchSettings.h
@@ -60,5 +60,5 @@ public:
 	UPROPERTY(EditAnywhere, Config, Category = "Search",
 		meta = (ClampMin = "0.1", ClampMax = "3.0",
 			ToolTip = "Maximum L2 distance for a result to be considered a semantic match. Lower values are stricter. Results must pass this threshold OR have a BM25 keyword match to appear."))
-	float VectorDistanceCutoff = 1.25f;
+	float VectorDistanceCutoff = 1.2f;
 };
