Fix the auto indexing on save  not working on new assets.

#Synced-CL 52518173

diff --git a/Engine/Plugins/Experimental/SemanticSearch/Source/SemanticSearch/Private/SemanticSearchModule.cpp b/Engine/Plugins/Experimental/SemanticSearch/Source/SemanticSearch/Private/SemanticSearchModule.cpp
--- a/Engine/Plugins/Experimental/SemanticSearch/Source/SemanticSearch/Private/SemanticSearchModule.cpp
+++ b/Engine/Plugins/Experimental/SemanticSearch/Source/SemanticSearch/Private/SemanticSearchModule.cpp
@@ -15,14 +15,15 @@
 #include "Settings/SemanticSearchSettings.h"
 #include "HAL/ConsoleManager.h"
 #include "Misc/Paths.h"
+#include "UObject/ObjectSaveContext.h"
+#include "UObject/Package.h"
 
 DEFINE_LOG_CATEGORY(LogSemanticSearch);
 
-bool UE::SemanticSearch::ISemanticSearchModule::IsInIndexedFolder(const FAssetData& AssetData)
+bool UE::SemanticSearch::ISemanticSearchModule::IsInIndexedFolder(FStringView PackagePath)
 {
 	const USemanticSearchSettings* Settings = USemanticSearchSettings::Get();
 	if (!Settings) { return false; }
-	const FString PackagePath = AssetData.PackageName.ToString();
 	for (const FDirectoryPath& Dir : Settings->IndexedPaths)
 	{
 		if (PackagePath.StartsWith(Dir.Path))
@@ -33,6 +34,11 @@ bool UE::SemanticSearch::ISemanticSearchModule::IsInIndexedFolder(const FAssetDa
 	return false;
 }
 
+bool UE::SemanticSearch::ISemanticSearchModule::IsInIndexedFolder(const FAssetData& AssetData)
+{
+	return IsInIndexedFolder(AssetData.PackageName.ToString());
+}
+
 namespace UE::SemanticSearch::Private
 {
 IMPLEMENT_MODULE(FSemanticSearchModule, SemanticSearch)
@@ -75,8 +81,8 @@ void FSemanticSearchModule::StartupModule()
 	IAssetRegistry& AssetRegistry = FModuleManager::LoadModuleChecked<FAssetRegistryModule>(TEXT("AssetRegistry")).Get();
 	AssetAddedDelegateHandle = AssetRegistry.OnAssetAdded().AddRaw(
 		this, &FSemanticSearchModule::OnAssetAdded);
-	AssetUpdatedOnDiskDelegateHandle = AssetRegistry.OnAssetUpdatedOnDisk().AddRaw(
-		this, &FSemanticSearchModule::OnAssetUpdatedOnDisk);
+	PackageSavedDelegateHandle = UPackage::PackageSavedWithContextEvent.AddRaw(
+		this, &FSemanticSearchModule::OnPackageSaved);
 	AssetRenamedDelegateHandle = AssetRegistry.OnAssetRenamed().AddRaw(
 		this, &FSemanticSearchModule::OnAssetRenamed);
 	AssetRemovedDelegateHandle = AssetRegistry.OnAssetRemoved().AddRaw(
@@ -96,6 +102,11 @@ void FSemanticSearchModule::ShutdownModule()
 	{
 		GEditor->GetTimerManager()->ClearTimer(DismissTimerHandle);
 		GEditor->GetTimerManager()->ClearTimer(PeriodicSaveTimerHandle);
+		for (TPair<FName, FTimerHandle>& Pair : PendingPackageSavedTimers)
+		{
+			GEditor->GetTimerManager()->ClearTimer(Pair.Value);
+		}
+		PendingPackageSavedTimers.Empty();
 	}
 
 	FCoreDelegates::OnEnginePreExit.Remove(PreExitDelegateHandle);
@@ -104,7 +115,7 @@ void FSemanticSearchModule::ShutdownModule()
 	{
 		IAssetRegistry& AssetRegistry = AssetRegistryModule->Get();
 		AssetRegistry.OnAssetAdded().Remove(AssetAddedDelegateHandle);
-		AssetRegistry.OnAssetUpdatedOnDisk().Remove(AssetUpdatedOnDiskDelegateHandle);
+		UPackage::PackageSavedWithContextEvent.Remove(PackageSavedDelegateHandle);
 		AssetRegistry.OnAssetRenamed().Remove(AssetRenamedDelegateHandle);
 		AssetRegistry.OnAssetRemoved().Remove(AssetRemovedDelegateHandle);
 		AssetRegistry.OnFilesLoaded().Remove(FilesLoadedDelegateHandle);
@@ -182,34 +193,88 @@ void FSemanticSearchModule::OnAssetAdded(const FAssetData& AssetData)
 	TriggerIndexing(AssetData, Settings && Settings->bBuildUncachedOnStartup);
 }
 
-void FSemanticSearchModule::OnAssetUpdatedOnDisk(const FAssetData& AssetData)
+void FSemanticSearchModule::OnPackageSaved(const FString& PackageFileName, UPackage* Package, FObjectPostSaveContext ObjectSaveContext)
 {
 	const USemanticSearchSettings* Settings = USemanticSearchSettings::Get();
 	if (!Settings || !Settings->bAutoIndexOnSave)
 	{
 		return;
 	}
-	if (!AssetData.IsRedirector() && ISemanticSearchModule::IsInIndexedFolder(AssetData))
+
+	if (!IsInIndexedFolder(Package->GetName()))
 	{
-		UE_LOG(LogSemanticSearch, Log, TEXT("OnAssetUpdatedOnDisk: %s"), *AssetData.GetObjectPathString());
-		TriggerIndexing(AssetData, /*bBuildOnMiss=*/true);
+		return;
 	}
+
+	FName PackageName = Package->GetFName();
+	FIoHash SavedHash = Package->GetSavedHash();
+	UE_LOG(LogSemanticSearch, Log, TEXT("OnPackageSaved: %s (deferring until AR matches hash)"), *PackageName.ToString());
+
+	// Defer so the AssetRegistry has time to update its data for the saved package.
+	// Timer callbacks run on the game thread via FTimerManager::Tick.
+	if (!GEditor)
+	{
+		return;
+	}
+	constexpr int32 InitialRetryCount = 0;
+	FTimerHandle& Handle = PendingPackageSavedTimers.FindOrAdd(PackageName);
+	GEditor->GetTimerManager()->SetTimer(Handle,
+		FTimerDelegate::CreateRaw(this, &FSemanticSearchModule::OnPackageSavedDeferred, PackageName, SavedHash, InitialRetryCount),
+		0.1f, /*bLoop=*/false);
 }
 
-void FSemanticSearchModule::OnAssetRenamed(const FAssetData& NewAssetData, const FString& OldObjectPath)
+void FSemanticSearchModule::OnPackageSavedDeferred(FName PackageName, FIoHash ExpectedHash, int32 RetryCount)
 {
-	const USemanticSearchSettings* Settings = USemanticSearchSettings::Get();
-	auto IsPathInIndexedFolders = [Settings](const FString& Path) -> bool
+	check(IsInGameThread());
+	static constexpr int32 MaxRetries = 10;
+
+	IAssetRegistry& AssetRegistry = FModuleManager::LoadModuleChecked<FAssetRegistryModule>(TEXT("AssetRegistry")).Get();
+
+	FAssetPackageData AssetPackageData;
+	UE::AssetRegistry::EExists Result = AssetRegistry.TryGetAssetPackageData(PackageName, AssetPackageData);
+	if (Result != UE::AssetRegistry::EExists::Exists || AssetPackageData.GetPackageSavedHash() != ExpectedHash)
 	{
-		if (!Settings) { return false; }
-		for (const FDirectoryPath& Dir : Settings->IndexedPaths)
+		if (RetryCount < MaxRetries)
 		{
-			if (Path.StartsWith(Dir.Path)) { return true; }
+			UE_LOG(LogSemanticSearch, Verbose, TEXT("OnPackageSavedDeferred: AR not up to date for %s yet, retrying (%d/%d)"),
+				*PackageName.ToString(), RetryCount + 1, MaxRetries);
+			if (GEditor)
+			{
+				FTimerHandle& Handle = PendingPackageSavedTimers.FindOrAdd(PackageName);
+				GEditor->GetTimerManager()->SetTimer(Handle,
+					FTimerDelegate::CreateRaw(this, &FSemanticSearchModule::OnPackageSavedDeferred, PackageName, ExpectedHash, RetryCount + 1),
+					0.1f, /*bLoop=*/false);
+			}
+		}
+		else
+		{
+			UE_LOG(LogSemanticSearch, Warning, TEXT("OnPackageSavedDeferred: AR hash never matched for %s after %d retries, giving up"),
+				*PackageName.ToString(), MaxRetries);
+			PendingPackageSavedTimers.Remove(PackageName);
 		}
-		return false;
-	};
-	const bool bOldInFolder = IsPathInIndexedFolders(OldObjectPath);
-	const bool bNewInFolder = ISemanticSearchModule::IsInIndexedFolder(NewAssetData);
+
+		return;
+	}
+
+	PendingPackageSavedTimers.Remove(PackageName);
+
+	TArray<FAssetData> Assets;
+	AssetRegistry.GetAssetsByPackageName(PackageName, Assets);
+	for (const FAssetData& AssetData : Assets)
+	{
+		if (IsInIndexedFolder(AssetData))
+		{
+			UE_LOG(LogSemanticSearch, Log, TEXT("OnPackageSavedDeferred: indexing %s"), *AssetData.GetObjectPathString());
+			TriggerIndexing(AssetData, /*bBuildOnMiss=*/true);
+		}
+	}
+}
+
+void FSemanticSearchModule::OnAssetRenamed(const FAssetData& NewAssetData, const FString& OldObjectPath)
+{
+	const USemanticSearchSettings* Settings = USemanticSearchSettings::Get();
+	const bool bOldInFolder = IsInIndexedFolder(OldObjectPath);
+	const bool bNewInFolder = IsInIndexedFolder(NewAssetData);
 
 	UE_LOG(LogSemanticSearch, Log, TEXT("OnAssetRenamed: %s -> %s (OldInFolder=%d, NewInFolder=%d)"),
 		*OldObjectPath, *NewAssetData.GetObjectPathString(), bOldInFolder, bNewInFolder);
diff --git a/Engine/Plugins/Experimental/SemanticSearch/Source/SemanticSearch/Private/SemanticSearchModule.h b/Engine/Plugins/Experimental/SemanticSearch/Source/SemanticSearch/Private/SemanticSearchModule.h
--- a/Engine/Plugins/Experimental/SemanticSearch/Source/SemanticSearch/Private/SemanticSearchModule.h
+++ b/Engine/Plugins/Experimental/SemanticSearch/Source/SemanticSearch/Private/SemanticSearchModule.h
@@ -12,6 +12,9 @@
 #include "IO/IoHash.h"
 #include "TimerManager.h"
 
+class UPackage;
+class FObjectPostSaveContext;
+
 DECLARE_LOG_CATEGORY_EXTERN(LogSemanticSearch, Log, All);
 
 namespace UE::SemanticSearch::Private
@@ -33,7 +36,8 @@ private:
 	void TriggerIndexing(const FAssetData& AssetData, bool bBuildOnMiss);
 
 	void OnAssetAdded(const FAssetData& AssetData);
-	void OnAssetUpdatedOnDisk(const FAssetData& AssetData);
+	void OnPackageSaved(const FString& PackageFileName, UPackage* Package, FObjectPostSaveContext ObjectSaveContext);
+	void OnPackageSavedDeferred(FName PackageName, FIoHash ExpectedHash, int32 RetryCount);
 	void OnAssetRenamed(const FAssetData& NewAssetData, const FString& OldObjectPath);
 	void OnAssetRemoved(const FAssetData& AssetData);
 	void OnFilesLoaded();
@@ -48,7 +52,7 @@ private:
 
 	TSharedPtr<IEmbeddingProvider> RegisteredEmbeddingProvider;
 	FDelegateHandle AssetAddedDelegateHandle;
-	FDelegateHandle AssetUpdatedOnDiskDelegateHandle;
+	FDelegateHandle PackageSavedDelegateHandle;
 	FDelegateHandle AssetRenamedDelegateHandle;
 	FDelegateHandle AssetRemovedDelegateHandle;
 	FDelegateHandle FilesLoadedDelegateHandle;
@@ -63,5 +67,6 @@ private:
 
 	/** Package hashes loaded from saved index, used for staleness checks during startup. */
 	TMap<int64, FIoHash> LoadedPackageHashes;
+	TMap<FName, FTimerHandle> PendingPackageSavedTimers;
 };
 }
diff --git a/Engine/Plugins/Experimental/SemanticSearch/Source/SemanticSearch/Public/ISemanticSearchModule.h b/Engine/Plugins/Experimental/SemanticSearch/Source/SemanticSearch/Public/ISemanticSearchModule.h
--- a/Engine/Plugins/Experimental/SemanticSearch/Source/SemanticSearch/Public/ISemanticSearchModule.h
+++ b/Engine/Plugins/Experimental/SemanticSearch/Source/SemanticSearch/Public/ISemanticSearchModule.h
@@ -30,5 +30,8 @@ namespace UE::SemanticSearch
 
 		/** Returns true if the asset is in a folder configured for indexing. */
 		SEMANTICSEARCH_API static bool IsInIndexedFolder(const FAssetData& AssetData);
+
+		/** Returns true if the given package path is in a folder configured for indexing. */
+		SEMANTICSEARCH_API static bool IsInIndexedFolder(FStringView PackagePath);
 	};
 }
\ No newline at end of file
