Change the path where save the cached semantic index to be a per project cache.

#Synced-CL 52518173

#Synced-CL 52518173

diff --git a/Engine/Plugins/Experimental/SemanticSearch/Source/SemanticSearch/Private/SemanticSearchModule.cpp b/Engine/Plugins/Experimental/SemanticSearch/Source/SemanticSearch/Private/SemanticSearchModule.cpp
--- a/Engine/Plugins/Experimental/SemanticSearch/Source/SemanticSearch/Private/SemanticSearchModule.cpp
+++ b/Engine/Plugins/Experimental/SemanticSearch/Source/SemanticSearch/Private/SemanticSearchModule.cpp
@@ -52,7 +52,7 @@ static FIoHash GetPackageSavedHash(const FAssetData& AssetData)
 
 FString FSemanticSearchModule::GetSavedIndexPath()
 {
-	return FPaths::EngineSavedDir() / TEXT("SemanticSearch") / TEXT("SearchIndex.bin");
+	return FPaths::ProjectSavedDir() / TEXT("SemanticSearch") / TEXT("SearchIndex.bin");
 }
 
 FSemanticSearchModule& FSemanticSearchModule::Get()
