Add a button to reindexe the missing assets in the semantic search Index.

#Synced-CL 52518173

diff --git a/Engine/Plugins/Experimental/SemanticSearch/Source/SemanticSearchEditorIntegrations/Private/ContentBrowser/SemanticSearchCB.cpp b/Engine/Plugins/Experimental/SemanticSearch/Source/SemanticSearchEditorIntegrations/Private/ContentBrowser/SemanticSearchCB.cpp
--- a/Engine/Plugins/Experimental/SemanticSearch/Source/SemanticSearchEditorIntegrations/Private/ContentBrowser/SemanticSearchCB.cpp
+++ b/Engine/Plugins/Experimental/SemanticSearch/Source/SemanticSearchEditorIntegrations/Private/ContentBrowser/SemanticSearchCB.cpp
@@ -22,7 +22,9 @@
 #include "ToolMenus.h"
 #include "Widgets/Input/SButton.h"
 #include "Widgets/Layout/SBox.h"
+#include "Widgets/SBoxPanel.h"
 #include "Widgets/Text/STextBlock.h"
+#include "AssetRegistry/IAssetRegistry.h"
 
 DEFINE_LOG_CATEGORY_STATIC(LogSemanticSearchCB, Log, All);
 
@@ -377,6 +379,15 @@ static FDelayedAutoRegisterHelper SemanticSearchStatsNavBarExtension(
 						return;
 					}
 
+					auto GetUnindexedCount = []() -> int32
+					{
+						const UE::SemanticSearch::FHybridSearchIndex& Index = UE::SemanticSearch::FHybridSearchIndex::Get();
+						const int32 Indexed = Index.GetIndexedCount();
+						const int32 Failed = Index.GetFailedCount();
+						const int32 TotalSupported = UE::SemanticSearch::FAssetProcessorManager::Get().GetSupportedAssetCount();
+						return FMath::Max(0, TotalSupported - Indexed - Failed);
+					};
+
 					FToolMenuEntry& Entry = InSection.AddEntry(
 						FToolMenuEntry::InitWidget(
 							"SemanticSearchStatsEntry",
@@ -389,30 +400,78 @@ static FDelayedAutoRegisterHelper SemanticSearchStatsNavBarExtension(
 									: EVisibility::Collapsed;
 							})
 							[
-								SNew(STextBlock)
-								.TextStyle(FAppStyle::Get(), "SmallText")
-								.Text_Lambda([]() -> FText
-								{
-									const UE::SemanticSearch::FHybridSearchIndex& Index = UE::SemanticSearch::FHybridSearchIndex::Get();
-									const int32 Indexed = Index.GetIndexedCount();
-									const int32 Failed = Index.GetFailedCount();
-									const int32 TotalSupported = UE::SemanticSearch::FAssetProcessorManager::Get().GetSupportedAssetCount();
-									const int32 Unindexed = FMath::Max(0, TotalSupported - Indexed - Failed);
-									if (Failed > 0)
+								SNew(SHorizontalBox)
+								+ SHorizontalBox::Slot()
+								.AutoWidth()
+								.VAlign(VAlign_Center)
+								[
+									SNew(STextBlock)
+									.TextStyle(FAppStyle::Get(), "SmallText")
+									.Text_Lambda([]() -> FText
 									{
+										const UE::SemanticSearch::FHybridSearchIndex& Index = UE::SemanticSearch::FHybridSearchIndex::Get();
+										const int32 Indexed = Index.GetIndexedCount();
+										const int32 Failed = Index.GetFailedCount();
+										const int32 TotalSupported = UE::SemanticSearch::FAssetProcessorManager::Get().GetSupportedAssetCount();
+										const int32 Unindexed = FMath::Max(0, TotalSupported - Indexed - Failed);
+										if (Failed > 0)
+										{
+											return FText::Format(
+												NSLOCTEXT("SemanticSearch", "IndexStatsWithFailed",
+													"Semantic Search Index: {0} indexed | {1} unindexed | {2} failed"),
+												FText::AsNumber(Indexed),
+												FText::AsNumber(Unindexed),
+												FText::AsNumber(Failed));
+										}
 										return FText::Format(
-											NSLOCTEXT("SemanticSearch", "IndexStatsWithFailed",
-												"Semantic Search Index: {0} indexed | {1} unindexed | {2} failed"),
+											NSLOCTEXT("SemanticSearch", "IndexStats",
+												"Semantic Search Index: {0} indexed | {1} unindexed"),
 											FText::AsNumber(Indexed),
-											FText::AsNumber(Unindexed),
-											FText::AsNumber(Failed));
-									}
-									return FText::Format(
-										NSLOCTEXT("SemanticSearch", "IndexStats",
-											"Semantic Search Index: {0} indexed | {1} unindexed"),
-										FText::AsNumber(Indexed),
-										FText::AsNumber(Unindexed));
-								})
+											FText::AsNumber(Unindexed));
+									})
+								]
+								+ SHorizontalBox::Slot()
+								.AutoWidth()
+								.VAlign(VAlign_Center)
+								.Padding(6.f, 0.f, 0.f, 0.f)
+								[
+									SNew(SButton)
+									.ButtonStyle(FAppStyle::Get(), "HoverOnlyButton")
+									.ContentPadding(FMargin(4.f, 0.f))
+									.ToolTipText(NSLOCTEXT("SemanticSearch", "ReindexTooltip", "Index all unindexed assets in configured folders"))
+									.Visibility_Lambda([GetUnindexedCount]()
+									{
+										return GetUnindexedCount() > 0
+											? EVisibility::Visible
+											: EVisibility::Collapsed;
+									})
+									.OnClicked_Lambda([]()
+									{
+										UE::SemanticSearch::FHybridSearchIndex& Index = UE::SemanticSearch::FHybridSearchIndex::Get();
+										UE::SemanticSearch::FAssetProcessorManager& ProcessorManager = UE::SemanticSearch::FAssetProcessorManager::Get();
+										UE::SemanticSearch::ISemanticSearchModule& Module = UE::SemanticSearch::ISemanticSearchModule::Get();
+										IAssetRegistry& AR = IAssetRegistry::GetChecked();
+
+										AR.EnumerateAllAssets([&](const FAssetData& Asset)
+										{
+											if (!Asset.IsRedirector()
+												&& UE::SemanticSearch::ISemanticSearchModule::IsInIndexedFolder(Asset)
+												&& ProcessorManager.GetProcessorForAsset(Asset)
+												&& !Index.Contains(Asset))
+											{
+												Module.IndexAsset(Asset);
+											}
+											return true;
+										});
+										return FReply::Handled();
+									})
+									[
+										SNew(STextBlock)
+										.TextStyle(FAppStyle::Get(), "SmallText")
+										.Text(NSLOCTEXT("SemanticSearch", "ReindexButton", "Reindex"))
+										.ColorAndOpacity(FSlateColor(FLinearColor(0.f, 0.8f, 1.f)))
+									]
+								]
 							],
 							FText::GetEmpty()));
 					Entry.WidgetData.StyleParams.VerticalAlignment = VAlign_Center;
