Minor UX improvements to the semantic search mode in the CB.

#Synced-CL 52518173

diff --git a/Engine/Plugins/Experimental/SemanticSearch/Source/SemanticSearchEditorIntegrations/Private/ContentBrowser/SemanticSearchCB.cpp b/Engine/Plugins/Experimental/SemanticSearch/Source/SemanticSearchEditorIntegrations/Private/ContentBrowser/SemanticSearchCB.cpp
--- a/Engine/Plugins/Experimental/SemanticSearch/Source/SemanticSearchEditorIntegrations/Private/ContentBrowser/SemanticSearchCB.cpp
+++ b/Engine/Plugins/Experimental/SemanticSearch/Source/SemanticSearchEditorIntegrations/Private/ContentBrowser/SemanticSearchCB.cpp
@@ -89,6 +89,8 @@ namespace Private
 			}
 		}
 
+		virtual bool IsSortOverridden() const override { return !ActiveQuery.IsEmpty(); }
+
 		virtual void SortItemList(TArrayView<TSharedPtr<FAssetViewItem>> Items) override
 		{
 
@@ -257,8 +259,7 @@ namespace Private
 			PublishResults(MakeArrayView(LastMatchedItems));
 		}
 
-		// When no query is active, surface every item that has asset data.
-		// Folder items and other non-asset entries are excluded.
+		// When no query is active, surface every item that is a folder or has asset data.
 		void PublishPassThrough()
 		{
 			TArray<TSharedPtr<FAssetViewItem>> PassedItems;
@@ -266,7 +267,8 @@ namespace Private
 			for (const TSharedPtr<FAssetViewItem>& Item : CurrentItems)
 			{
 				FAssetData AssetData;
-				if (Item->GetItem().Legacy_TryGetAssetData(AssetData))
+				const FContentBrowserItem& ItemRef = Item->GetItem();
+				if (ItemRef.IsFolder() || ItemRef.Legacy_TryGetAssetData(AssetData))
 				{
 					PassedItems.Add(Item);
 				}
diff --git a/Engine/Source/Editor/ContentBrowser/Private/SAssetView.cpp b/Engine/Source/Editor/ContentBrowser/Private/SAssetView.cpp
--- a/Engine/Source/Editor/ContentBrowser/Private/SAssetView.cpp
+++ b/Engine/Source/Editor/ContentBrowser/Private/SAssetView.cpp
@@ -4437,7 +4437,7 @@ void SAssetView::SortList(bool bSyncToSelection)
 	if ( !IsRenamingAsset() )
 	{
 		PRAGMA_DISABLE_EXPERIMENTAL_WARNINGS
-		if (AssetSearchOverride.IsValid() && AssetSearchOverride->IsOverrideActive())
+		if (AssetSearchOverride.IsValid() && AssetSearchOverride->IsOverrideActive() && AssetSearchOverride->IsSortOverridden())
 		{
 			AssetSearchOverride->SortItemList(FilteredAssetItems);
 		}
diff --git a/Engine/Source/Editor/ContentBrowser/Public/Experimental/ContentBrowserSearchOverride.h b/Engine/Source/Editor/ContentBrowser/Public/Experimental/ContentBrowserSearchOverride.h
--- a/Engine/Source/Editor/ContentBrowser/Public/Experimental/ContentBrowserSearchOverride.h
+++ b/Engine/Source/Editor/ContentBrowser/Public/Experimental/ContentBrowserSearchOverride.h
@@ -56,6 +56,12 @@ namespace UE::Editor::ContentBrowser::Extension
 
 		virtual void SortItemList(TArrayView<TSharedPtr<FAssetViewItem>> Items) = 0;
 
+		/**
+		 * Returns true if the override wants to handle sorting instead of the default
+		 * Content Browser sort manager. When false, SortItemList will not be called.
+		 */
+		virtual bool IsSortOverridden() const { return true; }
+
 		/**
 		 * Returns true while the override has async work in flight whose results
 		 * have not yet been published. The asset view will keep ticking while this is true.
