Fix support for quantized indexes. Add a tool menu dialog window to allow switching between indexes and (re)training PQ index.

- IQuantizedVectorIndex interface separated out from IVectorIndex: SupportsQuantization(), Quantize(), AddQuantized(),
GetCodebookHash(), SerializeCodebook()
- VectorIndexFactory: type-dispatched creation, deserialization, and codebook loading
- PQ codebook hashing
- DDC support for quantized embedding codes (keyed by asset + codebook hash)
- TriggerIndexing optimization for quantized embeddings: fetches caption first, then tries pre-quantized codes
from DDC before falling back to full float embedding fetch
- SwitchIndexType: calls SwitchToFlatIndex / SwitchToQuantizedIndex and reindexes assets using ddc fetch
- SwitchToQuantizedIndex loads existing codebook from disk when available, skipping training
- RetrainIndex: forces fresh codebook training for quantized indexes
- Index dialog: switch/retrain buttons, throbber during async operations, auto-refresh
via OnIndexChanged delegate when indexing completes
- Auto-save index on dismiss when indexing batch exceeds threshold
- Nav bar stats use GetIndexStats() for consistency

#Synced-CL 52562012

diff --git a/Engine/Plugins/Experimental/SemanticSearch/Source/SemanticSearch/Private/AssetProcessorManager.cpp b/Engine/Plugins/Experimental/SemanticSearch/Source/SemanticSearch/Private/AssetProcessorManager.cpp
--- a/Engine/Plugins/Experimental/SemanticSearch/Source/SemanticSearch/Private/AssetProcessorManager.cpp
+++ b/Engine/Plugins/Experimental/SemanticSearch/Source/SemanticSearch/Private/AssetProcessorManager.cpp
@@ -151,6 +151,21 @@ namespace UE::SemanticSearch
 		});
 	}
 
+	bool FAssetProcessorManager::GetQuantizedEmbeddingData(const FAssetData& InAssetData, const FIoHash& CodebookHash,
+		TUniqueFunction<void(FAssetQuantizedEmbeddingResult&&, FString&&)>&& InOnDataAvailable, bool bBuildOnMiss)
+	{
+		return GetOrCreateTask(InAssetData, [&InOnDataAvailable, &CodebookHash, bBuildOnMiss](const TSharedRef<Private::FSemanticSearchBuildCacheTask>& Task)
+		{
+			Task->GetQuantizedEmbeddingData(CodebookHash,
+				[OnDataAvailable = MoveTemp(InOnDataAvailable)](Private::FQuantizedEmbeddingDerivedData&& QuantizedData, FString&& ErrorMessage) mutable
+				{
+					FAssetQuantizedEmbeddingResult Result;
+					Result.QuantizedCodes = MoveTemp(QuantizedData.QuantizedEmbedding);
+					OnDataAvailable(MoveTemp(Result), MoveTemp(ErrorMessage));
+				}, bBuildOnMiss);
+		});
+	}
+
 	bool FAssetProcessorManager::GetEmbeddingData(const FAssetData& InAssetData, TUniqueFunction<void(FAssetEmbeddingResult&&, FString&&)>&& InOnDataAvailable, bool bBuildOnMiss)
 	{
 		return GetOrCreateTask(InAssetData, [&InOnDataAvailable, bBuildOnMiss](const TSharedRef<Private::FSemanticSearchBuildCacheTask>& Task)
diff --git a/Engine/Plugins/Experimental/SemanticSearch/Source/SemanticSearch/Private/HybridSearchIndex.cpp b/Engine/Plugins/Experimental/SemanticSearch/Source/SemanticSearch/Private/HybridSearchIndex.cpp
--- a/Engine/Plugins/Experimental/SemanticSearch/Source/SemanticSearch/Private/HybridSearchIndex.cpp
+++ b/Engine/Plugins/Experimental/SemanticSearch/Source/SemanticSearch/Private/HybridSearchIndex.cpp
@@ -5,15 +5,16 @@
 #include "SemanticSearchModule.h"
 #include "AssetRegistry/AssetData.h"
 #include "Hash/Blake3.h"
+#include "Interfaces/IPluginManager.h"
 #include "Misc/FileHelper.h"
 #include "Misc/Paths.h"
 #include "Misc/ScopeRWLock.h"
 #include "Serialization/MemoryReader.h"
 #include "Serialization/MemoryWriter.h"
+#include "Interfaces/IQuantizedVectorIndex.h"
 #include "Settings/SemanticSearchSettings.h"
 #include "TextIndex/BM25Index.h"
-#include "VectorIndex/FlatVectorIndex.h"
-#include "VectorIndex/PQVectorIndex.h"
+#include "VectorIndex/VectorIndexFactory.h"
 
 namespace UE::SemanticSearch
 {
@@ -51,13 +52,30 @@ void FHybridSearchIndex::EnsureInitialized(int32 EmbeddingDimension)
 
 	const USemanticSearchSettings* Settings = USemanticSearchSettings::Get();
 
-	if (Settings && Settings->IndexType == ESemanticSearchIndexType::PQ)
+	const ESemanticSearchIndexType Type = Settings ? Settings->IndexType : ESemanticSearchIndexType::Flat;
+	const FString CodebookPath = GetCodebookPath(Type);
+
+	if (!CodebookPath.IsEmpty())
 	{
-		VectorIndex = MakeShared<FPQVectorIndex>(EmbeddingDimension);
+		TArray<uint8> CodebookData;
+		if (FFileHelper::LoadFileToArray(CodebookData, *CodebookPath))
+		{
+			VectorIndex = DeserializeCodebook(Type, CodebookData, EmbeddingDimension);
+			if (VectorIndex)
+			{
+				UE_LOG(LogSemanticSearch, Log, TEXT("Loaded codebook from %s"), *CodebookPath);
+			}
+		}
+
+		if (!VectorIndex)
+		{
+			UE_LOG(LogSemanticSearch, Warning, TEXT("No valid codebook at %s, creating untrained index"), *CodebookPath);
+			VectorIndex = CreateVectorIndex(Type, EmbeddingDimension, Settings);
+		}
 	}
 	else
 	{
-		VectorIndex = MakeShared<FFlatVectorIndex>(EmbeddingDimension);
+		VectorIndex = CreateVectorIndex(Type, EmbeddingDimension, Settings);
 	}
 
 	bInitialized = true;
@@ -68,6 +86,16 @@ TSharedPtr<IVectorIndex> FHybridSearchIndex::GetVectorIndex() const
 	return VectorIndex;
 }
 
+void FHybridSearchIndex::SetVectorIndex(TSharedPtr<IVectorIndex> NewIndex)
+{
+	FWriteScopeLock WriteLock(IndexLock);
+	VectorIndex = MoveTemp(NewIndex);
+	if (VectorIndex)
+	{
+		bInitialized = true;
+	}
+}
+
 bool FHybridSearchIndex::IsHybridEnabled() const
 {
 	const USemanticSearchSettings* Settings = USemanticSearchSettings::Get();
@@ -116,8 +144,7 @@ void FHybridSearchIndex::Add(
 	const FAssetData& Asset,
 	TConstArrayView<float> Embedding,
 	FStringView Caption,
-	TConstArrayView<FString> Keywords,
-	TConstArrayView<uint8> QuantizedCodes)
+	TConstArrayView<FString> Keywords)
 {
 	const int64 ID = GetAssetIndexID(Asset);
 
@@ -127,17 +154,10 @@ void FHybridSearchIndex::Add(
 		EnsureInitialized(Embedding.Num());
 	}
 
-	// Add to vector index — use pre-quantized path if available
-	if (VectorIndex && VectorIndex->IsTrained() && Embedding.Num() > 0)
+	// Add to vector index
+	if (VectorIndex && Embedding.Num() > 0)
 	{
-		if (QuantizedCodes.Num() > 0 && VectorIndex->SupportsQuantization())
-		{
-			VectorIndex->AddQuantized(ID, QuantizedCodes);
-		}
-		else
-		{
-			VectorIndex->Add(ID, Embedding);
-		}
+		VectorIndex->Add(ID, Embedding);
 	}
 
 	// Add to BM25 index
@@ -153,6 +173,28 @@ void FHybridSearchIndex::Add(
 	}
 }
 
+void FHybridSearchIndex::AddQuantized(
+	const FAssetData& Asset,
+	TConstArrayView<uint8> QuantizedCodes,
+	FStringView Caption,
+	TConstArrayView<FString> Keywords)
+{
+	check(VectorIndex && VectorIndex->SupportsQuantization() && VectorIndex->IsTrained());
+	const int64 ID = GetAssetIndexID(Asset);
+
+	static_cast<IQuantizedVectorIndex*>(VectorIndex.Get())->AddQuantized(ID, QuantizedCodes);
+
+	if (IsHybridEnabled())
+	{
+		BM25Index->Add(ID, Asset.GetObjectPathString(), Caption, Keywords);
+	}
+
+	{
+		FScopeLock Lock(&FailedLock);
+		FailedAssets.Remove(ID);
+	}
+}
+
 void FHybridSearchIndex::Remove(const FAssetData& Asset)
 {
 	Remove(Asset.GetObjectPathString());
@@ -305,6 +347,23 @@ TArray<FHybridSearchResult> FHybridSearchIndex::Search(
 	return Results;
 }
 
+FString FHybridSearchIndex::GetCodebookPath(ESemanticSearchIndexType Type)
+{
+	const FString Filename = GetCodebookFilename(Type);
+	if (Filename.IsEmpty())
+	{
+		return FString();
+	}
+
+	TSharedPtr<IPlugin> Plugin = IPluginManager::Get().FindPlugin(TEXT("SemanticSearch"));
+	if (Plugin.IsValid())
+	{
+		return Plugin->GetBaseDir() / TEXT("Resources") / Filename;
+	}
+	// Fallback: hardcoded relative to plugin base dir
+	return FPaths::EnginePluginsDir() / TEXT("Experimental") / TEXT("SemanticSearch") / TEXT("Resources") / Filename;
+}
+
 bool FHybridSearchIndex::SaveToFile(const FString& FilePath, TMap<int64, FIoHash>& PackageHashes) const
 {
 	FReadScopeLock ReadLock(IndexLock);
@@ -319,13 +378,20 @@ bool FHybridSearchIndex::SaveToFile(const FString& FilePath, TMap<int64, FIoHash
 	bool bWriteInitialized = bInitialized;
 	Ar << bWriteInitialized;
 
-	// Vector index type + data
+	// Vector index type + codebook hash
 	{
 		const USemanticSearchSettings* Settings = USemanticSearchSettings::Get();
 		uint8 IndexType = (bInitialized && VectorIndex && Settings)
 			? static_cast<uint8>(Settings->IndexType)
 			: static_cast<uint8>(ESemanticSearchIndexType::Flat);
 		Ar << IndexType;
+
+		FIoHash CodebookHash;
+		if (bInitialized && VectorIndex && VectorIndex->SupportsQuantization() && VectorIndex->IsTrained())
+		{
+			CodebookHash = static_cast<IQuantizedVectorIndex*>(VectorIndex.Get())->GetCodebookHash();
+		}
+		Ar << CodebookHash;
 	}
 
 	if (bInitialized && VectorIndex)
@@ -364,8 +430,6 @@ bool FHybridSearchIndex::LoadFromFile(const FString& FilePath, TMap<int64, FIoHa
 		return false;
 	}
 
-	FWriteScopeLock WriteLock(IndexLock);
-
 	FMemoryReader Ar(FileData);
 
 	uint32 Magic = 0;
@@ -383,6 +447,9 @@ bool FHybridSearchIndex::LoadFromFile(const FString& FilePath, TMap<int64, FIoHa
 	uint8 IndexType = 0;
 	Ar << IndexType;
 
+	FIoHash SavedCodebookHash;
+	Ar << SavedCodebookHash;
+
 	int32 VectorDim = 0;
 	TArray<uint8> VectorData;
 	Ar << VectorDim;
@@ -407,15 +474,29 @@ bool FHybridSearchIndex::LoadFromFile(const FString& FilePath, TMap<int64, FIoHa
 			return false;
 		}
 
-		if (SavedType == ESemanticSearchIndexType::PQ)
-		{
-			LoadedVectorIndex = TSharedPtr<IVectorIndex>(FPQVectorIndex::Deserialize(VectorData, VectorDim).Release());
-		}
-		else
+		// Validate codebook hash for compressed index types
+		if (SavedType != ESemanticSearchIndexType::Flat && SavedCodebookHash != FIoHash())
 		{
-			LoadedVectorIndex = TSharedPtr<IVectorIndex>(FFlatVectorIndex::Deserialize(VectorData, VectorDim).Release());
+			const FString CodebookPath = GetCodebookPath(SavedType);
+			TArray<uint8> CodebookData;
+			if (FFileHelper::LoadFileToArray(CodebookData, *CodebookPath))
+			{
+				TSharedPtr<IVectorIndex> TempIndex = DeserializeCodebook(SavedType, CodebookData, VectorDim);
+				if (TempIndex && TempIndex->SupportsQuantization())
+				{
+					FIoHash CurrentCodebookHash = static_cast<IQuantizedVectorIndex*>(TempIndex.Get())->GetCodebookHash();
+					if (SavedCodebookHash != CurrentCodebookHash)
+					{
+						UE_LOG(LogSemanticSearch, Warning,
+							TEXT("Saved index codebook hash does not match current codebook, discarding saved index"));
+						return false;
+					}
+				}
+			}
 		}
 
+		LoadedVectorIndex = DeserializeVectorIndex(SavedType, VectorData, VectorDim);
+
 		if (!LoadedVectorIndex)
 		{
 			return false;
@@ -439,6 +520,7 @@ bool FHybridSearchIndex::LoadFromFile(const FString& FilePath, TMap<int64, FIoHa
 	}
 
 	// All deserialization succeeded — swap into members atomically
+	FWriteScopeLock WriteLock(IndexLock);
 	bInitialized = LoadedInitialized;
 	VectorIndex = MoveTemp(LoadedVectorIndex);
 	BM25Index = MoveTemp(LoadedBM25);
diff --git a/Engine/Plugins/Experimental/SemanticSearch/Source/SemanticSearch/Private/SemanticSearchDerivedDataBuilder.cpp b/Engine/Plugins/Experimental/SemanticSearch/Source/SemanticSearch/Private/SemanticSearchDerivedDataBuilder.cpp
--- a/Engine/Plugins/Experimental/SemanticSearch/Source/SemanticSearch/Private/SemanticSearchDerivedDataBuilder.cpp
+++ b/Engine/Plugins/Experimental/SemanticSearch/Source/SemanticSearch/Private/SemanticSearchDerivedDataBuilder.cpp
@@ -2,12 +2,14 @@
 
 #include "SemanticSearchDerivedDataBuilder.h"
 
+#include "AssetProcessorManager.h"
 #include "Containers/SharedString.h"
 #include "HybridSearchIndex.h"
 #include "ProfilingDebugging/CpuProfilerTrace.h"
 #include "Containers/StringFwd.h"
 #include "DerivedDataCache.h"
 #include "Interfaces/IAssetProcessor.h"
+#include "Interfaces/IQuantizedVectorIndex.h"
 #include "Interfaces/IVectorIndex.h"
 #include "Misc/Guid.h"
 #include "Misc/ScopeRWLock.h"
@@ -182,6 +184,54 @@ void FSemanticSearchBuildCacheTask::GetCaptionData(TUniqueFunction<void(FCaption
 		});
 }
 
+void FSemanticSearchBuildCacheTask::GetQuantizedEmbeddingData(
+	const FIoHash& CodebookHash,
+	TUniqueFunction<void(FQuantizedEmbeddingDerivedData&&, FString&&)>&& InOnDataAvailable,
+	bool bBuildOnMiss)
+{
+	TRACE_CPUPROFILER_EVENT_SCOPE_STR("SemanticSearch::GetQuantizedEmbeddingData");
+
+	{
+		QueuedQuantizedRequests.Enqueue(MoveTemp(InOnDataAvailable));
+		bool bPreviousValue = false;
+		HasActiveQuantizedRequest.compare_exchange_strong(bPreviousValue, true);
+		if (bPreviousValue)
+		{
+			return;
+		}
+	}
+
+	DerivedData::FCacheKey QuantizedKey = GenerateQuantizedCacheKey(AssetKey, CodebookHash);
+
+	using namespace DerivedData;
+
+	FCacheGetChunkRequest Request;
+	Request.Name = DDCRequestName;
+	Request.Key = QuantizedKey;
+	Request.Id = GetQuantizedEmbeddingDatachunkId();
+	Request.Policy = GetCachePolicy();
+
+	FRequestBarrier Barrier(*this);
+	GetCache().GetChunks({ Request }, *this,
+		[this, CodebookHash, bBuildOnMiss](FCacheGetChunkResponse&& Response) mutable
+		{
+			if (Response.Status == EStatus::Ok && Response.RawData.GetSize() > 0)
+			{
+				TRACE_CPUPROFILER_EVENT_SCOPE_STR("SemanticSearch::GetQuantizedEmbeddingData::DDCHit");
+				FQuantizedEmbeddingDerivedData Data;
+				FMemoryReaderView Reader(Response.RawData.GetView(), true);
+				SerializeQuantizedEmbedding(Data, Reader);
+				Response.RawData.Reset();
+				BroadcastQueuedQuantizedRequest(MoveTemp(Data), FString());
+			}
+			else
+			{
+				TRACE_CPUPROFILER_EVENT_SCOPE_STR("SemanticSearch::GetQuantizedEmbeddingData::DDCMiss");
+				BuildQuantizedDerivedData(CodebookHash, bBuildOnMiss);
+			}
+		});
+}
+
 void FSemanticSearchBuildCacheTask::CacheData()
 {
 	TRACE_CPUPROFILER_EVENT_SCOPE_STR("SemanticSearch::CacheData");
@@ -232,6 +282,23 @@ DerivedData::FCacheKey FSemanticSearchBuildCacheTask::GenerateCacheKey(const FAs
 	return CacheKey;
 }
 
+DerivedData::FCacheKey FSemanticSearchBuildCacheTask::GenerateQuantizedCacheKey(
+	const DerivedData::FCacheKey& BaseCacheKey, const FIoHash& CodebookHash)
+{
+	FMemoryHasherBlake3 HashWriter;
+
+	// Combine base key hash with codebook hash
+	FIoHash BaseHash = BaseCacheKey.Hash;
+	HashWriter << BaseHash;
+	FIoHash MutableCodebookHash = CodebookHash;
+	HashWriter << MutableCodebookHash;
+
+	DerivedData::FCacheKey QuantizedKey;
+	QuantizedKey.Hash = HashWriter.Finalize();
+	QuantizedKey.Bucket = BaseCacheKey.Bucket;
+	return QuantizedKey;
+}
+
 DerivedData::FCacheBucket FSemanticSearchBuildCacheTask::GetCacheKeyBucket(const TSharedRef<IAssetProcessor>& InAssetProcessor)
 {
 	TStringBuilder<256> Builder;
@@ -278,16 +345,6 @@ void FSemanticSearchBuildCacheTask::BuildDerivedData()
 											FEmbeddingDerivedData EmbeddingDerivedData;
 											EmbeddingDerivedData.Embedding = Embedding.Embedding;
 
-											// Quantize the embedding if the active index supports it (PQ).
-											FQuantizedEmbeddingDerivedData QuantizedEmbeddingData;
-											if (TSharedPtr<IVectorIndex> VectorIndex = FHybridSearchIndex::Get().GetVectorIndex())
-											{
-												if (VectorIndex->SupportsQuantization() && VectorIndex->IsTrained())
-												{
-													QuantizedEmbeddingData.QuantizedEmbedding = VectorIndex->Quantize(EmbeddingDerivedData.Embedding);
-												}
-											}
-
 											if (AssetKey == GenerateCacheKey(*Asset, AssetProcessor))
 											{
 												using namespace DerivedData;
@@ -310,15 +367,6 @@ void FSemanticSearchBuildCacheTask::BuildDerivedData()
 													RecordBuilder.AddValue(GetEmbeddingDatachunkId(), MakeSharedBufferFromArray(MoveTemp(Buffer)));
 												}
 
-												// Save quantized embedding (PQ codes) if available
-												if (QuantizedEmbeddingData.QuantizedEmbedding.Num() > 0)
-												{
-													TArray<uint8> Buffer;
-													FMemoryWriter Writer(Buffer);
-													SerializeQuantizedEmbedding(QuantizedEmbeddingData, Writer);
-													RecordBuilder.AddValue(GetQuantizedEmbeddingDatachunkId(), MakeSharedBufferFromArray(MoveTemp(Buffer)));
-												}
-
 												DerivedData::FCachePutRequest PutRequest{ DDCRequestName, RecordBuilder.Build(), GetCachePolicy() | ECachePolicy::SkipData };
 
 												DerivedData::GetCache().Put(
@@ -542,4 +590,118 @@ void FSemanticSearchBuildCacheTask::BroadcastQueuedRequests(FEmbeddingDerivedDat
 	}
 }
 
+TUniqueFunction<void(FQuantizedEmbeddingDerivedData&&, FString&&)> FSemanticSearchBuildCacheTask::NoLock_GetQuantizedRequests()
+{
+	TArray<TUniqueFunction<void(FQuantizedEmbeddingDerivedData&&, FString&&)>, TInlineAllocator<4>> Callbacks;
+	bool bPreviousValue = true;
+	HasActiveQuantizedRequest.compare_exchange_strong(bPreviousValue, false);
+	if (bPreviousValue)
+	{
+		TUniqueFunction<void(FQuantizedEmbeddingDerivedData&&, FString&&)> Callback;
+		while (QueuedQuantizedRequests.Dequeue(Callback))
+		{
+			Callbacks.Add(MoveTemp(Callback));
+		}
+	}
+
+	if (Callbacks.IsEmpty())
+	{
+		return {};
+	}
+
+	return [Callbacks = MoveTemp(Callbacks)](FQuantizedEmbeddingDerivedData&& Data, FString&& ErrorMessage) mutable
+	{
+		FQuantizedEmbeddingDerivedData CopiedData;
+		FQuantizedEmbeddingDerivedData* DataToUse = &Data;
+		FString CopiedErrorMessage;
+		FString* ErrorMessageToUse = &ErrorMessage;
+		int32 Count = 0;
+
+		for (TUniqueFunction<void(FQuantizedEmbeddingDerivedData&&, FString&&)>& Callback : Callbacks)
+		{
+			++Count;
+			if (Callback)
+			{
+				bool bSwapCurrentToCopied = Count < Callbacks.Num();
+				if (bSwapCurrentToCopied)
+				{
+					CopiedData = *DataToUse;
+					CopiedErrorMessage = *ErrorMessageToUse;
+				}
+
+				Callback(MoveTemp(*DataToUse), MoveTemp(*ErrorMessageToUse));
+
+				if (bSwapCurrentToCopied)
+				{
+					DataToUse = &CopiedData;
+					ErrorMessageToUse = &CopiedErrorMessage;
+				}
+				else
+				{
+					break;
+				}
+			}
+		}
+	};
+}
+
+void FSemanticSearchBuildCacheTask::BroadcastQueuedQuantizedRequest(FQuantizedEmbeddingDerivedData&& QuantizedData, FString&& ErrorMessage)
+{
+	TRACE_CPUPROFILER_EVENT_SCOPE_STR("SemanticSearch::BroadcastQuantizedRequest");
+	TUniqueFunction<void(FQuantizedEmbeddingDerivedData&&, FString&&)> QuantizedCallback = NoLock_GetQuantizedRequests();
+
+	if (QuantizedCallback)
+	{
+		QuantizedCallback(MoveTemp(QuantizedData), MoveTemp(ErrorMessage));
+	}
+}
+
+void FSemanticSearchBuildCacheTask::BuildQuantizedDerivedData(const FIoHash& CodebookHash, bool bBuildOnMiss)
+{
+	TRACE_CPUPROFILER_EVENT_SCOPE_STR("SemanticSearch::BuildQuantizedDerivedData");
+
+	// Get the float embedding (DDC hit in most cases, or triggers full BuildDerivedData if bBuildOnMiss)
+	GetEmbeddingData([this, CodebookHash](FEmbeddingDerivedData&& EmbeddingData, FString&& Error)
+		{
+			if (!Error.IsEmpty() || EmbeddingData.Embedding.IsEmpty())
+			{
+				BroadcastQueuedQuantizedRequest(FQuantizedEmbeddingDerivedData(), MoveTemp(Error));
+				return;
+			}
+
+			// Quantize using active index
+			TSharedPtr<IVectorIndex> VectorIndex = FHybridSearchIndex::Get().GetVectorIndex();
+			if (!VectorIndex || !VectorIndex->SupportsQuantization() || !VectorIndex->IsTrained())
+			{
+				BroadcastQueuedQuantizedRequest(FQuantizedEmbeddingDerivedData(), FString(TEXT("No trained quantized index")));
+				return;
+			}
+
+			auto* QuantizedIdx = static_cast<IQuantizedVectorIndex*>(VectorIndex.Get());
+			FQuantizedEmbeddingDerivedData QuantizedData;
+			QuantizedData.QuantizedEmbedding = QuantizedIdx->Quantize(EmbeddingData.Embedding);
+
+			// Write to DDC
+			using namespace DerivedData;
+			FCacheKey QuantizedKey = GenerateQuantizedCacheKey(AssetKey, CodebookHash);
+			FCacheRecordBuilder RecordBuilder(QuantizedKey);
+			TArray<uint8> Buffer;
+			FMemoryWriter Writer(Buffer);
+			SerializeQuantizedEmbedding(QuantizedData, Writer);
+			RecordBuilder.AddValue(GetQuantizedEmbeddingDatachunkId(), MakeSharedBufferFromArray(MoveTemp(Buffer)));
+
+			FCachePutRequest PutRequest{ DDCRequestName, RecordBuilder.Build(), GetCachePolicy() | ECachePolicy::SkipData };
+			GetCache().Put({ PutRequest }, *this,
+				[AssetPtr = Asset](FCachePutResponse&& Response)
+				{
+					if (Response.Status != DerivedData::EStatus::Ok)
+					{
+						UE_LOGF(LogSemanticSearch, Warning, "Failed to store quantized codes in DDC for: %ls", *AssetPtr->GetObjectPathString());
+					}
+				});
+
+			BroadcastQueuedQuantizedRequest(MoveTemp(QuantizedData), FString());
+		}, bBuildOnMiss);
+}
+
 }
diff --git a/Engine/Plugins/Experimental/SemanticSearch/Source/SemanticSearch/Private/SemanticSearchDerivedDataBuilder.h b/Engine/Plugins/Experimental/SemanticSearch/Source/SemanticSearch/Private/SemanticSearchDerivedDataBuilder.h
--- a/Engine/Plugins/Experimental/SemanticSearch/Source/SemanticSearch/Private/SemanticSearchDerivedDataBuilder.h
+++ b/Engine/Plugins/Experimental/SemanticSearch/Source/SemanticSearch/Private/SemanticSearchDerivedDataBuilder.h
@@ -10,6 +10,7 @@
 #include "DerivedDataCacheKey.h"
 #include "DerivedDataRequestOwner.h"
 #include "HAL/CriticalSection.h"
+#include "IO/IoHash.h"
 #include "Templates/SharedPointer.h"
 #include <atomic>
 
@@ -52,20 +53,29 @@ public:
 	void GetEmbeddingData(TUniqueFunction<void (FEmbeddingDerivedData&& EmbeddingData, FString&& ErrorMessage)>&& InOnDataAvailable, bool bBuildOnMiss = true);
 	void GetCaptionData(TUniqueFunction<void (FCaptionDerivedData&& CaptionData, FString&& ErrorMessage)>&& InOnDataAvailable);
 
+	/** Retrieve quantized embedding data from DDC. On miss, builds by fetching the float embedding and quantizing. */
+	void GetQuantizedEmbeddingData(const FIoHash& CodebookHash,
+		TUniqueFunction<void(FQuantizedEmbeddingDerivedData&&, FString&&)>&& InOnDataAvailable,
+		bool bBuildOnMiss = true);
+
 	void CacheData();
 
 	static DerivedData::FCacheKey GenerateCacheKey(const FAssetData& InAssetData, const TSharedRef<IAssetProcessor>& InAssetProcessor);
+	static DerivedData::FCacheKey GenerateQuantizedCacheKey(const DerivedData::FCacheKey& BaseCacheKey, const FIoHash& CodebookHash);
 	static DerivedData::FCacheBucket GetCacheKeyBucket(const TSharedRef<IAssetProcessor>& InAssetProcessor);
 
 private:
 
 	void BuildDerivedData();
+	void BuildQuantizedDerivedData(const FIoHash& CodebookHash, bool bBuildOnMiss);
 	void BroadcastQueuedRequests(FEmbeddingDerivedData&& EmbeddingData, FCaptionDerivedData&& CaptionData, FString&& ErrorMessage);
 	void BroadcastQueuedCaptionRequest(FCaptionDerivedData&& CaptionData, FString&& ErrorMessage);
 	void BroadcastQueuedEmbeddingRequest(FEmbeddingDerivedData&& EmbeddingData, FString&& ErrorMessage);
+	void BroadcastQueuedQuantizedRequest(FQuantizedEmbeddingDerivedData&& QuantizedData, FString&& ErrorMessage);
 
 	TUniqueFunction<void(FCaptionDerivedData&&, FString&&)> NoLock_GetCaptionRequests();
 	TUniqueFunction<void(FEmbeddingDerivedData&&, FString&&)> NoLock_GetEmbeddingRequest();
+	TUniqueFunction<void(FQuantizedEmbeddingDerivedData&&, FString&&)> NoLock_GetQuantizedRequests();
 
 
 	TSharedRef<const FAssetData> Asset;
@@ -78,8 +88,10 @@ private:
 	bool bIsBuildingDerivedData = false;
 	TQueue<TUniqueFunction<void (FEmbeddingDerivedData&&, FString&&)>, EQueueMode::Mpsc> QueuedEmbeddingRequests;
 	TQueue<TUniqueFunction<void (FCaptionDerivedData&&, FString&&)>, EQueueMode::Mpsc> QueuedCaptionRequests;
+	TQueue<TUniqueFunction<void (FQuantizedEmbeddingDerivedData&&, FString&&)>, EQueueMode::Mpsc> QueuedQuantizedRequests;
 	std::atomic<bool> HasActiveEmbeddingRequest = false;
 	std::atomic<bool> HasActiveCaptionRequest = false;
+	std::atomic<bool> HasActiveQuantizedRequest = false;
 };
 }
 }
diff --git a/Engine/Plugins/Experimental/SemanticSearch/Source/SemanticSearch/Private/SemanticSearchModule.cpp b/Engine/Plugins/Experimental/SemanticSearch/Source/SemanticSearch/Private/SemanticSearchModule.cpp
--- a/Engine/Plugins/Experimental/SemanticSearch/Source/SemanticSearch/Private/SemanticSearchModule.cpp
+++ b/Engine/Plugins/Experimental/SemanticSearch/Source/SemanticSearch/Private/SemanticSearchModule.cpp
@@ -6,12 +6,17 @@
 #include "AssetProcessors/AssetProcessorUtils.h"
 #include "HybridSearchIndex.h"
 #include "TextIndex/BM25Index.h"
+#include "Interfaces/IQuantizedVectorIndex.h"
+#include "VectorIndex/VectorIndexFactory.h"
 
 #include "Modules/ModuleManager.h"
 #include "RemoteEmbeddingProvider.h"
 #include "Async/Async.h"
 #include "AssetRegistry/AssetRegistryModule.h"
 #include "Editor.h"
+#include "Interfaces/IPluginManager.h"
+#include "Misc/ConfigCacheIni.h"
+#include "Misc/FileHelper.h"
 #include "Settings/SemanticSearchSettings.h"
 #include "HAL/ConsoleManager.h"
 #include "Misc/Paths.h"
@@ -124,47 +129,98 @@ void FSemanticSearchModule::ShutdownModule()
 
 void FSemanticSearchModule::TriggerIndexing(const FAssetData& AssetData, bool bBuildOnMiss)
 {
+	if (!FAssetProcessorManager::Get().GetProcessorForAsset(AssetData))
+	{
+		return;
+	}
 	OnIndexingStarted();
 
-	bool bStarted = FAssetProcessorManager::Get().GetEmbeddingData(AssetData,
-		[Asset = FAssetData(AssetData), bBuildOnMiss](FAssetEmbeddingResult&& EmbeddingResult, FString&& Error) mutable
+	// Check up front if we're on a quantized index
+	FHybridSearchIndex& Index = FHybridSearchIndex::Get();
+	TSharedPtr<IVectorIndex> VectorIndex = Index.GetVectorIndex();
+	const bool bQuantized = VectorIndex && VectorIndex->SupportsQuantization() && VectorIndex->IsTrained();
+
+	// Step 1: Fetch caption (always needed for both flat and quantized paths)
+	bool bStarted = FAssetProcessorManager::Get().GetCaptionData(AssetData,
+		[Asset = FAssetData(AssetData), bBuildOnMiss, bQuantized](FAssetCaptionResult&& CaptionResult, FString&& CaptionError) mutable
 		{
-			if (!Error.IsEmpty() || EmbeddingResult.Embedding.IsEmpty())
+			if (bQuantized)
 			{
-				FHybridSearchIndex::Get().MarkFailed(GetAssetIndexID(Asset));
-				if (bBuildOnMiss)
+				// Step 2a (quantized): Get quantized codes (DDC hit or build on miss)
+				FHybridSearchIndex& Index = FHybridSearchIndex::Get();
+				TSharedPtr<IVectorIndex> VectorIndex = Index.GetVectorIndex();
+				if (VectorIndex && VectorIndex->SupportsQuantization() && VectorIndex->IsTrained())
 				{
-					UE_LOG(LogSemanticSearch, Warning, TEXT("Indexing failed for %s: %s"),
-						*Asset.GetObjectPathString(),
-						Error.IsEmpty() ? TEXT("empty embedding") : *Error);
+					auto* QuantizedIndex = static_cast<IQuantizedVectorIndex*>(VectorIndex.Get());
+					FIoHash CodebookHash = QuantizedIndex->GetCodebookHash();
+
+					bool bStarted = FAssetProcessorManager::Get().GetQuantizedEmbeddingData(Asset, CodebookHash,
+						[Asset, CaptionResult = MoveTemp(CaptionResult)](
+							FAssetQuantizedEmbeddingResult&& QuantizedResult, FString&& QuantizedError) mutable
+						{
+							if (QuantizedResult.QuantizedCodes.Num() > 0)
+							{
+								FHybridSearchIndex& Index = FHybridSearchIndex::Get();
+								Index.Remove(Asset);
+								Index.AddQuantized(Asset, QuantizedResult.QuantizedCodes, CaptionResult.Caption, CaptionResult.Keywords);
+
+								UE_LOG(LogSemanticSearch, Log, TEXT("Indexed %s (quantized, index size: %lld)"),
+									*Asset.GetObjectPathString(),
+									Index.GetVectorIndex() ? Index.GetVectorIndex()->GetCount() : 0);
+								FSemanticSearchModule::Get().OnIndexingFinished(true);
+							}
+							else
+							{
+								FHybridSearchIndex::Get().MarkFailed(GetAssetIndexID(Asset));
+								UE_LOG(LogSemanticSearch, Warning, TEXT("Failed to get quantized codes for %s: %s"),
+									*Asset.GetObjectPathString(), *QuantizedError);
+								FSemanticSearchModule::Get().OnIndexingFinished(false);
+							}
+						}, bBuildOnMiss);
+					if (!bStarted)
+					{
+						FSemanticSearchModule::Get().OnIndexingFinished(false);
+					}
+					return;
 				}
-				FSemanticSearchModule::Get().OnIndexingFinished(false);
-				return;
 			}
 
-			FAssetProcessorManager::Get().GetCaptionData(Asset,
-				[Asset, Embedding = MoveTemp(EmbeddingResult.Embedding)](FAssetCaptionResult&& CaptionResult, FString&& CaptionError) mutable
+			// Step 2b (flat or index changed): fetch embedding and add
+			bool bStarted = FAssetProcessorManager::Get().GetEmbeddingData(Asset,
+				[Asset, CaptionResult = MoveTemp(CaptionResult), bBuildOnMiss](
+					FAssetEmbeddingResult&& EmbeddingResult, FString&& Error) mutable
 				{
+					if (!Error.IsEmpty() || EmbeddingResult.Embedding.IsEmpty())
+					{
+						FHybridSearchIndex::Get().MarkFailed(GetAssetIndexID(Asset));
+						if (bBuildOnMiss)
+						{
+							UE_LOG(LogSemanticSearch, Warning, TEXT("Indexing failed for %s: %s"),
+								*Asset.GetObjectPathString(),
+								Error.IsEmpty() ? TEXT("empty embedding") : *Error);
+						}
+						FSemanticSearchModule::Get().OnIndexingFinished(false);
+						return;
+					}
+
 					FHybridSearchIndex& Index = FHybridSearchIndex::Get();
 					Index.Remove(Asset);
-					Index.Add(Asset, Embedding, CaptionResult.Caption, CaptionResult.Keywords);
+					Index.Add(Asset, EmbeddingResult.Embedding, CaptionResult.Caption, CaptionResult.Keywords);
+
 					UE_LOG(LogSemanticSearch, Log, TEXT("Indexed %s (%d dims, index size: %lld)"),
-						*Asset.GetObjectPathString(), Embedding.Num(),
+						*Asset.GetObjectPathString(), EmbeddingResult.Embedding.Num(),
 						Index.GetVectorIndex() ? Index.GetVectorIndex()->GetCount() : 0);
 					FSemanticSearchModule::Get().OnIndexingFinished(true);
-				});
-		}, bBuildOnMiss);
+				}, bBuildOnMiss);
+			if (!bStarted)
+			{
+				FSemanticSearchModule::Get().OnIndexingFinished(false);
+			}
+		});
 
 	if (!bStarted)
 	{
-		// Unsupported asset type
-		--IndexingTotal;
-		UpdateProgressNotification();
-
-		if (IndexingCompleted + IndexingFailed >= IndexingTotal)
-		{
-			DismissProgressNotification();
-		}
+		FSemanticSearchModule::Get().OnIndexingFinished(false);
 	}
 }
 
@@ -412,9 +468,17 @@ void FSemanticSearchModule::DismissProgressNotification()
 		FSlateNotificationManager::Get().CancelProgressNotification(ProgressHandle);
 		ProgressHandle = FProgressNotificationHandle();
 	}
+
+	if (IndexingCompleted >= SaveIndexThreshold)
+	{
+		SaveIndex();
+	}
+
 	IndexingTotal = 0;
 	IndexingCompleted = 0;
 	IndexingFailed = 0;
+
+	FHybridSearchIndex::Get().OnIndexChanged.Broadcast(true);
 }
 
 void FSemanticSearchModule::OnFilesLoaded()
@@ -442,6 +506,262 @@ void FSemanticSearchModule::OnFilesLoaded()
 	}
 }
 
+FSemanticSearchIndexStats FSemanticSearchModule::GetIndexStats() const
+{
+	FSemanticSearchIndexStats Stats;
+
+	const USemanticSearchSettings* Settings = USemanticSearchSettings::Get();
+	Stats.IndexType = Settings ? Settings->IndexType : ESemanticSearchIndexType::Flat;
+
+	FHybridSearchIndex& Index = FHybridSearchIndex::Get();
+	TSharedPtr<IVectorIndex> VectorIndex = Index.GetVectorIndex();
+	if (VectorIndex)
+	{
+		Stats.VectorCount = static_cast<int32>(VectorIndex->GetCount());
+		Stats.Dimension = VectorIndex->GetDimension();
+		Stats.bIsTrained = VectorIndex->IsTrained();
+		Stats.EstimatedVectorMemoryBytes = VectorIndex->EstimateMemoryBytes();
+	}
+
+	Stats.BM25Count = static_cast<int32>(Index.GetBM25Index().GetCount());
+	Stats.EstimatedBM25MemoryBytes = Index.GetBM25Index().EstimateMemoryBytes();
+	Stats.FailedCount = Index.GetFailedCount();
+	Stats.SupportedAssetCount = FAssetProcessorManager::Get().GetSupportedAssetCount();
+
+	return Stats;
+}
+
+void FSemanticSearchModule::SwitchIndexType(ESemanticSearchIndexType NewType)
+{
+	if (bIsSwitchingIndexType)
+	{
+		UE_LOG(LogSemanticSearch, Warning, TEXT("SwitchIndexType: Already switching, ignoring request"));
+		return;
+	}
+
+	const USemanticSearchSettings* Settings = USemanticSearchSettings::Get();
+	if (!Settings || Settings->IndexType == NewType)
+	{
+		return;
+	}
+
+	if (NewType == ESemanticSearchIndexType::Flat)
+	{
+		SwitchToFlatIndex();
+	}
+	else
+	{
+		SwitchToQuantizedIndex(NewType);
+	}
+}
+
+void FSemanticSearchModule::SwitchToFlatIndex()
+{
+	const USemanticSearchSettings* Settings = USemanticSearchSettings::Get();
+	TSharedPtr<IVectorIndex> CurrentIndex = FHybridSearchIndex::Get().GetVectorIndex();
+	const int32 Dimension = CurrentIndex ? CurrentIndex->GetDimension() : 0;
+
+	// Create new empty flat index and swap it in
+	TSharedPtr<IVectorIndex> NewIndex = CreateVectorIndex(ESemanticSearchIndexType::Flat, Dimension, Settings);
+	FHybridSearchIndex::Get().SetVectorIndex(NewIndex);
+
+	USemanticSearchSettings* MutableSettings = GetMutableDefault<USemanticSearchSettings>();
+	MutableSettings->IndexType = ESemanticSearchIndexType::Flat;
+	MutableSettings->SaveConfig();
+
+	UE_LOG(LogSemanticSearch, Log, TEXT("Switched to Flat index, re-indexing all assets from DDC..."));
+	bIsSwitchingIndexType = false;
+
+	// Re-index all assets — TriggerIndexing fetches embeddings from DDC
+	IndexAllAssets();
+}
+
+void FSemanticSearchModule::SwitchToQuantizedIndex(ESemanticSearchIndexType NewType, bool bForceTrain)
+{
+	const USemanticSearchSettings* Settings = USemanticSearchSettings::Get();
+	FHybridSearchIndex& Index = FHybridSearchIndex::Get();
+	TSharedPtr<IVectorIndex> CurrentIndex = Index.GetVectorIndex();
+	const int32 Dimension = CurrentIndex ? CurrentIndex->GetDimension() : 0;
+
+	if (Dimension == 0)
+	{
+		UE_LOG(LogSemanticSearch, Warning, TEXT("SwitchToQuantizedIndex: Current index dimension zero, cannot switch"));
+		bIsSwitchingIndexType = false;
+		FHybridSearchIndex::Get().OnIndexChanged.Broadcast(false);
+		return;
+	}
+
+	// Try to load existing codebook to skip training (unless forced)
+	TSharedPtr<IVectorIndex> NewIndex;
+	if (!bForceTrain)
+	{
+		const FString CodebookPath = FHybridSearchIndex::GetCodebookPath(NewType);
+		TArray<uint8> CodebookData;
+		if (FFileHelper::LoadFileToArray(CodebookData, *CodebookPath))
+		{
+			NewIndex = DeserializeCodebook(NewType, CodebookData, Dimension);
+			if (NewIndex)
+			{
+				UE_LOG(LogSemanticSearch, Log, TEXT("Loaded existing codebook from %s, skipping training"), *CodebookPath);
+			}
+		}
+	}
+
+	if (NewIndex)
+	{
+		// Codebook loaded — swap in the trained index and re-index from DDC
+		Index.SetVectorIndex(NewIndex);
+
+		USemanticSearchSettings* MutableSettings = GetMutableDefault<USemanticSearchSettings>();
+		MutableSettings->IndexType = NewType;
+		MutableSettings->SaveConfig();
+
+		UE_LOG(LogSemanticSearch, Log, TEXT("Switched to quantized index with existing codebook, re-indexing from DDC..."));
+		bIsSwitchingIndexType = false;
+
+		// TriggerIndexing will use the quantized fast path (DDC codes) when available
+		IndexAllAssets();
+	}
+	else
+	{
+		// No codebook — need to train. Collect vectors from current index for training.
+		if (!CurrentIndex || CurrentIndex->GetCount() == 0)
+		{
+			UE_LOG(LogSemanticSearch, Warning, TEXT("SwitchToQuantizedIndex: No vectors available for training"));
+			bIsSwitchingIndexType = false;
+			FHybridSearchIndex::Get().OnIndexChanged.Broadcast(false);
+			return;
+		}
+
+		const int64 Count = CurrentIndex->GetCount();
+		TArray<float> TrainingVectors;
+		TrainingVectors.Reserve(Count * Dimension);
+		TArray<int64> AllIDs;
+		AllIDs.Reserve(Count);
+
+		IAssetRegistry& AssetRegistry = FModuleManager::LoadModuleChecked<FAssetRegistryModule>(TEXT("AssetRegistry")).Get();
+		TArray<FAssetData> AllAssets;
+		AssetRegistry.GetAllAssets(AllAssets, /*bIncludeOnlyOnDiskAssets=*/true);
+
+		for (const FAssetData& Asset : AllAssets)
+		{
+			if (Asset.IsRedirector() || !IsInIndexedFolder(Asset))
+			{
+				continue;
+			}
+
+			const int64 ID = GetAssetIndexID(Asset);
+			if (!CurrentIndex->Contains(ID))
+			{
+				continue;
+			}
+
+			TArray<float> Embedding;
+			if (CurrentIndex->TryGetEmbedding(ID, Embedding))
+			{
+				AllIDs.Add(ID);
+				TrainingVectors.Append(Embedding);
+			}
+		}
+
+		if (AllIDs.Num() == 0)
+		{
+			UE_LOG(LogSemanticSearch, Warning, TEXT("SwitchToQuantizedIndex: Could not retrieve any vectors for training"));
+			bIsSwitchingIndexType = false;
+			FHybridSearchIndex::Get().OnIndexChanged.Broadcast(false);
+			return;
+		}
+
+		NewIndex = CreateVectorIndex(NewType, Dimension, Settings);
+
+		UE_LOG(LogSemanticSearch, Log, TEXT("Training quantized index (type=%d) with %d vectors (dim=%d)..."),
+			static_cast<uint8>(NewType), AllIDs.Num(), Dimension);
+
+		FProgressNotificationHandle SwitchProgress = FSlateNotificationManager::Get().StartProgressNotification(
+			NSLOCTEXT("SemanticSearch", "TrainingIndex", "Training quantized index..."), 0);
+
+		// Train on background thread, then swap and re-index from DDC on game thread
+		Async(EAsyncExecution::ThreadPool,
+			[this, NewType, NewIndex, TrainingVectors = MoveTemp(TrainingVectors), NumVectors = AllIDs.Num(), SwitchProgress]() mutable
+			{
+				NewIndex->Train(TrainingVectors, NumVectors);
+
+				auto* QuantizedIndex = static_cast<IQuantizedVectorIndex*>(NewIndex.Get());
+
+				// Save codebook
+				const FString CodebookPath = FHybridSearchIndex::GetCodebookPath(NewType);
+				if (!CodebookPath.IsEmpty())
+				{
+					TArray<uint8> CodebookData = QuantizedIndex->SerializeCodebook();
+					FString CodebookDir = FPaths::GetPath(CodebookPath);
+					IFileManager::Get().MakeDirectory(*CodebookDir, /*bTree=*/true);
+					if (FFileHelper::SaveArrayToFile(CodebookData, *CodebookPath))
+					{
+						UE_LOG(LogSemanticSearch, Log, TEXT("Saved codebook to %s"), *CodebookPath);
+					}
+					else
+					{
+						UE_LOG(LogSemanticSearch, Warning, TEXT("Failed to save codebook to %s"), *CodebookPath);
+					}
+				}
+
+				// Finalize on game thread — swap index and re-index all assets from DDC
+				AsyncTask(ENamedThreads::GameThread,
+					[this, NewType, NewIndex = MoveTemp(NewIndex), NumVectors, SwitchProgress]() mutable
+					{
+						FHybridSearchIndex::Get().SetVectorIndex(NewIndex);
+
+						USemanticSearchSettings* MutableSettings = GetMutableDefault<USemanticSearchSettings>();
+						MutableSettings->IndexType = NewType;
+						MutableSettings->SaveConfig();
+
+						FSlateNotificationManager::Get().CancelProgressNotification(SwitchProgress);
+						bIsSwitchingIndexType = false;
+
+						UE_LOG(LogSemanticSearch, Log, TEXT("Quantized index trained, re-indexing %d assets from DDC..."), NumVectors);
+
+						IndexAllAssets();
+					});
+			});
+	}
+}
+
+void FSemanticSearchModule::IndexAllAssets()
+{
+	IAssetRegistry& AssetRegistry = FModuleManager::LoadModuleChecked<FAssetRegistryModule>(TEXT("AssetRegistry")).Get();
+	TArray<FAssetData> AllAssets;
+	AssetRegistry.GetAllAssets(AllAssets, /*bIncludeOnlyOnDiskAssets=*/true);
+
+	const USemanticSearchSettings* Settings = USemanticSearchSettings::Get();
+
+	for (const FAssetData& Asset : AllAssets)
+	{
+		if (Asset.IsRedirector() || !IsInIndexedFolder(Asset))
+		{
+			continue;
+		}
+		TriggerIndexing(Asset, Settings && Settings->bBuildUncachedOnStartup);
+	}
+}
+
+void FSemanticSearchModule::RetrainIndex()
+{
+	const USemanticSearchSettings* Settings = USemanticSearchSettings::Get();
+	if (!Settings || Settings->IndexType == ESemanticSearchIndexType::Flat)
+	{
+		UE_LOG(LogSemanticSearch, Warning, TEXT("RetrainIndex: Current index is Flat, nothing to retrain"));
+		return;
+	}
+
+	if (bIsSwitchingIndexType)
+	{
+		UE_LOG(LogSemanticSearch, Warning, TEXT("RetrainIndex: Already switching/retraining, ignoring request"));
+		return;
+	}
+
+	SwitchToQuantizedIndex(Settings->IndexType, /*bForceTrain=*/true);
+}
+
 void FSemanticSearchModule::SaveIndex()
 {
 	FHybridSearchIndex& Index = FHybridSearchIndex::Get();
diff --git a/Engine/Plugins/Experimental/SemanticSearch/Source/SemanticSearch/Private/SemanticSearchModule.h b/Engine/Plugins/Experimental/SemanticSearch/Source/SemanticSearch/Private/SemanticSearchModule.h
--- a/Engine/Plugins/Experimental/SemanticSearch/Source/SemanticSearch/Private/SemanticSearchModule.h
+++ b/Engine/Plugins/Experimental/SemanticSearch/Source/SemanticSearch/Private/SemanticSearchModule.h
@@ -31,9 +31,15 @@ public:
 
 	virtual void RegisterEmbeddingProvider(TUniquePtr<IEmbeddingProvider>&& EmbeddingProvider) override;
 	virtual void IndexAsset(const FAssetData& AssetData) override;
+	virtual FSemanticSearchIndexStats GetIndexStats() const override;
+	virtual void SwitchIndexType(ESemanticSearchIndexType NewType) override;
+	virtual void RetrainIndex() override;
 
 private:
+	void SwitchToFlatIndex();
+	void SwitchToQuantizedIndex(ESemanticSearchIndexType NewType, bool bForceTrain = false);
 	void TriggerIndexing(const FAssetData& AssetData, bool bBuildOnMiss);
+	void IndexAllAssets();
 
 	void OnAssetAdded(const FAssetData& AssetData);
 	void OnPackageSaved(const FString& PackageFileName, UPackage* Package, FObjectPostSaveContext ObjectSaveContext);
@@ -65,6 +71,12 @@ private:
 	FTimerHandle DismissTimerHandle;
 	FTimerHandle PeriodicSaveTimerHandle;
 
+	/** Guards against concurrent SwitchIndexType calls. Only accessed on game thread. */
+	bool bIsSwitchingIndexType = false;
+
+	/** Minimum number of completed indexing operations before auto-saving on dismiss. */
+	static constexpr int32 SaveIndexThreshold = 100;
+
 	/** Package hashes loaded from saved index, used for staleness checks during startup. */
 	TMap<int64, FIoHash> LoadedPackageHashes;
 	TMap<FName, FTimerHandle> PendingPackageSavedTimers;
diff --git a/Engine/Plugins/Experimental/SemanticSearch/Source/SemanticSearch/Private/SemanticSearchPrivateTestCommands.cpp b/Engine/Plugins/Experimental/SemanticSearch/Source/SemanticSearch/Private/SemanticSearchPrivateTestCommands.cpp
new file mode 100644
--- /dev/null
+++ b/Engine/Plugins/Experimental/SemanticSearch/Source/SemanticSearch/Private/SemanticSearchPrivateTestCommands.cpp
@@ -0,0 +1,54 @@
+// Copyright Epic Games, Inc. All Rights Reserved.
+
+#include "SemanticSearchModule.h"
+#include "TextIndex/BM25Index.h"
+#include "HybridSearchIndex.h"
+#include "AssetProcessorManager.h"
+#include "AssetProcessors/AssetProcessorUtils.h"
+#include "ISemanticSearchModule.h"
+#include "RemoteEmbeddingProvider.h"
+#include "HAL/ConsoleManager.h"
+#include "Misc/FileHelper.h"
+#include "Misc/Paths.h"
+#include "AssetRegistry/AssetRegistryModule.h"
+
+namespace UE::SemanticSearch::Private
+{
+
+static FAutoConsoleCommand BM25IndexStatsCommand(
+	// TEXT("SemanticSearch.BM25IndexStats"),
+	// TEXT("Print BM25 index memory usage and posting list stats."),
+	// FConsoleCommandDelegate::CreateLambda([]()
+	// {
+	// 	FBM25Index& BM25Index = FHybridSearchIndex::Get().GetBM25Index();
+	// 	FBM25Index::FStats Stats = BM25Index.GetStats();
+
+	// 	UE_LOG(LogSemanticSearch, Log, TEXT("=== BM25 Index Stats ==="));
+	// 	UE_LOG(LogSemanticSearch, Log, TEXT("Documents:           %lld"), Stats.NumDocuments);
+	// 	UE_LOG(LogSemanticSearch, Log, TEXT("Unique tokens:       %lld"), Stats.NumUniqueTokens);
+	// 	UE_LOG(LogSemanticSearch, Log, TEXT("Total token count:   %lld"), Stats.TotalTokenCount);
+	// 	UE_LOG(LogSemanticSearch, Log, TEXT("Total posting entries: %lld"), Stats.TotalPostingEntries);
+	// 	UE_LOG(LogSemanticSearch, Log, TEXT("Posting list size:   min=%d  max=%d  avg=%.1f"),
+	// 		Stats.MinPostingListSize, Stats.MaxPostingListSize, Stats.AvgPostingListSize);
+	// 	UE_LOG(LogSemanticSearch, Log, TEXT("Estimated memory:    %.2f KB (%.2f MB)"),
+	// 		Stats.EstimatedMemoryBytes / 1024.0, Stats.EstimatedMemoryBytes / (1024.0 * 1024.0));
+	// })
+);
+
+static FAutoConsoleCommand BuildJobsStatsCommand(
+	// TEXT("SemanticSearch.BuildJobsStats"),
+	// TEXT("Print the number of cached build tasks and their memory usage."),
+	// FConsoleCommandDelegate::CreateLambda([]()
+	// {
+	// 	FAssetProcessorManager& Manager = FAssetProcessorManager::Get();
+	// 	const int32 NumJobs = Manager.GetNumBuildJobs();
+	// 	const SIZE_T AllocatedSize = Manager.GetBuildJobsAllocatedSize();
+
+	// 	UE_LOG(LogSemanticSearch, Log, TEXT("=== BuildsJobs Stats ==="));
+	// 	UE_LOG(LogSemanticSearch, Log, TEXT("Cached tasks:    %d"), NumJobs);
+	// 	UE_LOG(LogSemanticSearch, Log, TEXT("TStripedMap alloc: %.2f KB (%.2f MB)"),
+	// 		AllocatedSize / 1024.0, AllocatedSize / (1024.0 * 1024.0));
+	// })
+);
+
+}
diff --git a/Engine/Plugins/Experimental/SemanticSearch/Source/SemanticSearch/Private/TextIndex/BM25Index.cpp b/Engine/Plugins/Experimental/SemanticSearch/Source/SemanticSearch/Private/TextIndex/BM25Index.cpp
--- a/Engine/Plugins/Experimental/SemanticSearch/Source/SemanticSearch/Private/TextIndex/BM25Index.cpp
+++ b/Engine/Plugins/Experimental/SemanticSearch/Source/SemanticSearch/Private/TextIndex/BM25Index.cpp
@@ -382,6 +382,32 @@ int64 FBM25Index::GetCount() const
 	return DocLengths.Num();
 }
 
+int64 FBM25Index::EstimateMemoryBytes() const
+{
+	FReadScopeLock ReadLock(IndexLock);
+
+	int64 Bytes = 0;
+
+	// PostingLists: TMap overhead + per-entry arrays
+	Bytes += PostingLists.GetAllocatedSize();
+	for (const auto& [Hash, Postings] : PostingLists)
+	{
+		Bytes += Postings.GetAllocatedSize();
+	}
+
+	// DocLengths
+	Bytes += DocLengths.GetAllocatedSize();
+
+	// DocTokens: TMap overhead + per-doc token arrays
+	Bytes += DocTokens.GetAllocatedSize();
+	for (const auto& [ID, Tokens] : DocTokens)
+	{
+		Bytes += Tokens.GetAllocatedSize();
+	}
+
+	return Bytes;
+}
+
 bool FBM25Index::Contains(int64 ID) const
 {
 	FReadScopeLock ReadLock(IndexLock);
diff --git a/Engine/Plugins/Experimental/SemanticSearch/Source/SemanticSearch/Private/TextIndex/BM25Index.h b/Engine/Plugins/Experimental/SemanticSearch/Source/SemanticSearch/Private/TextIndex/BM25Index.h
--- a/Engine/Plugins/Experimental/SemanticSearch/Source/SemanticSearch/Private/TextIndex/BM25Index.h
+++ b/Engine/Plugins/Experimental/SemanticSearch/Source/SemanticSearch/Private/TextIndex/BM25Index.h
@@ -51,6 +51,9 @@ public:
 	/** Deserialize a BM25 index from an archive. Returns null on version mismatch or corrupt data. */
 	static TUniquePtr<FBM25Index> Deserialize(FArchive& Ar);
 
+	/** Estimate the memory footprint of this index in bytes. */
+	int64 EstimateMemoryBytes() const;
+
 private:
 	struct FPostingEntry
 	{
diff --git a/Engine/Plugins/Experimental/SemanticSearch/Source/SemanticSearch/Private/VectorIndex/FlatVectorIndex.cpp b/Engine/Plugins/Experimental/SemanticSearch/Source/SemanticSearch/Private/VectorIndex/FlatVectorIndex.cpp
--- a/Engine/Plugins/Experimental/SemanticSearch/Source/SemanticSearch/Private/VectorIndex/FlatVectorIndex.cpp
+++ b/Engine/Plugins/Experimental/SemanticSearch/Source/SemanticSearch/Private/VectorIndex/FlatVectorIndex.cpp
@@ -123,6 +123,15 @@ bool FFlatVectorIndex::TryGetEmbedding(int64 ID, TArray<float>& OutEmbedding) co
 	return true;
 }
 
+int64 FFlatVectorIndex::EstimateMemoryBytes() const
+{
+	FReadScopeLock ReadLock(IndexLock);
+	const int64 Count = IndexWithIDs->ntotal;
+	// Raw vectors: Count * Dimension * sizeof(float)
+	// ID overhead: ~24 bytes/entry for id_map + rev_map hash entry
+	return Count * Dimension * sizeof(float) + Count * 24;
+}
+
 TArray<uint8> FFlatVectorIndex::Serialize() const
 {
 	FReadScopeLock ReadLock(IndexLock);
diff --git a/Engine/Plugins/Experimental/SemanticSearch/Source/SemanticSearch/Private/VectorIndex/FlatVectorIndex.h b/Engine/Plugins/Experimental/SemanticSearch/Source/SemanticSearch/Private/VectorIndex/FlatVectorIndex.h
--- a/Engine/Plugins/Experimental/SemanticSearch/Source/SemanticSearch/Private/VectorIndex/FlatVectorIndex.h
+++ b/Engine/Plugins/Experimental/SemanticSearch/Source/SemanticSearch/Private/VectorIndex/FlatVectorIndex.h
@@ -40,6 +40,7 @@ public:
 	virtual TArray<uint8> Serialize() const override;
 	virtual bool Contains(int64 ID) const override;
 	virtual bool TryGetEmbedding(int64 ID, TArray<float>& OutEmbedding) const override;
+	virtual int64 EstimateMemoryBytes() const override;
 
 	/** Deserialize a full flat index. */
 	static TUniquePtr<FFlatVectorIndex> Deserialize(TConstArrayView<uint8> Data, int32 ExpectedDimension);
diff --git a/Engine/Plugins/Experimental/SemanticSearch/Source/SemanticSearch/Private/VectorIndex/PQVectorIndex.cpp b/Engine/Plugins/Experimental/SemanticSearch/Source/SemanticSearch/Private/VectorIndex/PQVectorIndex.cpp
--- a/Engine/Plugins/Experimental/SemanticSearch/Source/SemanticSearch/Private/VectorIndex/PQVectorIndex.cpp
+++ b/Engine/Plugins/Experimental/SemanticSearch/Source/SemanticSearch/Private/VectorIndex/PQVectorIndex.cpp
@@ -2,6 +2,8 @@
 
 #include "PQVectorIndex.h"
 
+#include "Containers/Set.h"
+#include "Hash/Blake3.h"
 #include "Misc/ScopeRWLock.h"
 
 THIRD_PARTY_INCLUDES_START
@@ -20,10 +22,10 @@ FPQVectorIndex::FPQVectorIndex(int32 InDimension, int32 InSubvectorSize, int32 I
 {
 	check(Dimension % SubvectorSize == 0);
 	faiss::IndexPQ* PQIndex = new faiss::IndexPQ(Dimension, NumSubquantizers, NBits);
-	IndexWithIDs = new faiss::IndexIDMap(PQIndex);
+	IndexWithIDs = new faiss::IndexIDMap2(PQIndex);
 }
 
-FPQVectorIndex::FPQVectorIndex(int32 InDimension, int32 InSubvectorSize, int32 InNBits, faiss::IndexIDMap* InIndex)
+FPQVectorIndex::FPQVectorIndex(int32 InDimension, int32 InSubvectorSize, int32 InNBits, faiss::IndexIDMap2* InIndex)
 	: Dimension(InDimension)
 	, SubvectorSize(InSubvectorSize)
 	, NumSubquantizers(InDimension / InSubvectorSize)
@@ -42,6 +44,8 @@ void FPQVectorIndex::Train(TConstArrayView<float> Vectors, int64 NumVectors)
 	check(Vectors.Num() == NumVectors * Dimension);
 	FWriteScopeLock WriteLock(IndexLock);
 	IndexWithIDs->train(NumVectors, Vectors.GetData());
+	// Invalidate cached codebook hash since training changed the codebook
+	bCodebookHashCached = false;
 }
 
 bool FPQVectorIndex::IsTrained() const
@@ -84,46 +88,123 @@ void FPQVectorIndex::AddQuantized(TConstArrayView<int64> IDs, TConstArrayView<ui
 
 	FWriteScopeLock WriteLock(IndexLock);
 
+	const int64 BasePosition = PQIndex->ntotal;
+
 	// Insert codes directly into the PQ index's internal storage
 	PQIndex->codes.insert(PQIndex->codes.end(), Codes.GetData(), Codes.GetData() + Codes.Num());
 	PQIndex->ntotal += Count;
 
-	// Insert IDs into the IDMap
+	// Insert IDs into the IDMap and rev_map
 	static_assert(sizeof(faiss::idx_t) == sizeof(int64), "FAISS idx_t size mismatch");
 	const faiss::idx_t* IDsPtr = reinterpret_cast<const faiss::idx_t*>(IDs.GetData());
 	IndexWithIDs->id_map.insert(IndexWithIDs->id_map.end(), IDsPtr, IDsPtr + Count);
 	IndexWithIDs->ntotal += Count;
+
+	// Update rev_map for O(1) lookups
+	for (int64 i = 0; i < Count; ++i)
+	{
+		IndexWithIDs->rev_map[IDsPtr[i]] = BasePosition + i;
+	}
+}
+
+bool FPQVectorIndex::Contains(int64 ID) const
+{
+	FReadScopeLock ReadLock(IndexLock);
+	return IndexWithIDs->rev_map.count(static_cast<faiss::idx_t>(ID)) > 0;
+}
+
+bool FPQVectorIndex::TryGetEmbedding(int64 ID, TArray<float>& OutEmbedding) const
+{
+	FReadScopeLock ReadLock(IndexLock);
+	if (IndexWithIDs->rev_map.count(static_cast<faiss::idx_t>(ID)) == 0)
+	{
+		return false;
+	}
+
+	OutEmbedding.SetNumUninitialized(Dimension);
+	// Returns approximate embedding decoded from PQ codes (not exact original)
+	IndexWithIDs->reconstruct(static_cast<faiss::idx_t>(ID), OutEmbedding.GetData());
+	return true;
+}
+
+bool FPQVectorIndex::Update(int64 OldID, int64 NewID)
+{
+	FWriteScopeLock WriteLock(IndexLock);
+
+	auto It = IndexWithIDs->rev_map.find(static_cast<faiss::idx_t>(OldID));
+	if (It == IndexWithIDs->rev_map.end())
+	{
+		return false;
+	}
+
+	const int64 Position = It->second;
+
+	// Update id_map entry
+	IndexWithIDs->id_map[Position] = static_cast<faiss::idx_t>(NewID);
+
+	// Rebuild rev_map to reflect the change
+	IndexWithIDs->rev_map.erase(It);
+	IndexWithIDs->rev_map[static_cast<faiss::idx_t>(NewID)] = Position;
+
+	return true;
+}
+
+int64 FPQVectorIndex::EstimateMemoryBytes() const
+{
+	FReadScopeLock ReadLock(IndexLock);
+	const faiss::IndexPQ* PQIndex = static_cast<const faiss::IndexPQ*>(IndexWithIDs->index);
+	const int64 Count = IndexWithIDs->ntotal;
+	const int64 CodeSize = PQIndex->pq.code_size;
+
+	// PQ codes storage
+	const int64 CodesBytes = Count * CodeSize;
+
+	// Codebook: M * ksub * dsub * sizeof(float)
+	const int64 M = PQIndex->pq.M;
+	const int64 Ksub = 1LL << PQIndex->pq.nbits;
+	const int64 Dsub = PQIndex->pq.dsub;
+	const int64 CodebookBytes = M * Ksub * Dsub * sizeof(float);
+
+	// ID map + rev_map overhead (~24 bytes per entry: 8 for id_map + ~16 for rev_map hash entry)
+	const int64 IDOverhead = Count * 24;
+
+	return CodesBytes + CodebookBytes + IDOverhead;
 }
 
 TArray<FSearchResult> FPQVectorIndex::Search(TConstArrayView<float> QueryVector, int32 K, TConstArrayView<int64> IDFilter) const
 {
 	check(QueryVector.Num() == Dimension);
 
+	FReadScopeLock ReadLock(IndexLock);
+
+	// IndexPQ::search does not support IDSelector (throws on non-SearchParametersPQ params).
+	// Work around by searching all vectors and post-filtering.
+	const bool bHasFilter = IDFilter.Num() > 0;
+	const int32 FetchK = bHasFilter ? static_cast<int32>(IndexWithIDs->ntotal) : K;
+	if (FetchK == 0) { return {}; }
+
 	TArray<float> Distances;
 	TArray<faiss::idx_t> Labels;
-	Distances.SetNumUninitialized(K);
-	Labels.SetNumUninitialized(K);
+	Distances.SetNumUninitialized(FetchK);
+	Labels.SetNumUninitialized(FetchK);
 
-	FReadScopeLock ReadLock(IndexLock);
+	IndexWithIDs->search(1, QueryVector.GetData(), FetchK, Distances.GetData(), Labels.GetData());
 
-	if (IDFilter.Num() > 0)
+	TSet<int64> FilterSet;
+	if (bHasFilter)
 	{
-		static_assert(sizeof(faiss::idx_t) == sizeof(int64), "FAISS idx_t size mismatch");
-		faiss::IDSelectorBatch Selector(IDFilter.Num(), reinterpret_cast<const faiss::idx_t*>(IDFilter.GetData()));
-		faiss::SearchParameters Params;
-		Params.sel = &Selector;
-		IndexWithIDs->search(1, QueryVector.GetData(), K, Distances.GetData(), Labels.GetData(), &Params);
-	}
-	else
-	{
-		IndexWithIDs->search(1, QueryVector.GetData(), K, Distances.GetData(), Labels.GetData());
+		FilterSet.Reserve(IDFilter.Num());
+		for (int64 ID : IDFilter)
+		{
+			FilterSet.Add(ID);
+		}
 	}
 
 	TArray<FSearchResult> Results;
 	Results.Reserve(K);
-	for (int32 i = 0; i < K; ++i)
+	for (int32 i = 0; i < FetchK && Results.Num() < K; ++i)
 	{
-		if (Labels[i] >= 0)
+		if (Labels[i] >= 0 && (!bHasFilter || FilterSet.Contains(static_cast<int64>(Labels[i]))))
 		{
 			FSearchResult Result;
 			Result.ID = static_cast<int64>(Labels[i]);
@@ -184,6 +265,28 @@ TArray<uint8> FPQVectorIndex::SerializeCodebook() const
 	return Result;
 }
 
+FIoHash FPQVectorIndex::GetCodebookHash() const
+{
+	{
+		FReadScopeLock ReadLock(IndexLock);
+		if (bCodebookHashCached)
+		{
+			return CachedCodebookHash;
+		}
+	}
+
+	TArray<uint8> CodebookData = SerializeCodebook();
+
+	FBlake3 Hasher;
+	Hasher.Update(CodebookData.GetData(), CodebookData.Num());
+	FBlake3Hash Hash = Hasher.Finalize();
+
+	FWriteScopeLock WriteLock(IndexLock);
+	CachedCodebookHash = FIoHash(Hash);
+	bCodebookHashCached = true;
+	return CachedCodebookHash;
+}
+
 TUniquePtr<FPQVectorIndex> FPQVectorIndex::DeserializeCodebook(TConstArrayView<uint8> Data, int32 ExpectedDimension)
 {
 	faiss::VectorIOReader Reader;
@@ -205,7 +308,7 @@ TUniquePtr<FPQVectorIndex> FPQVectorIndex::DeserializeCodebook(TConstArrayView<u
 	PQIndex->is_trained = true;
 	delete PQ;
 
-	faiss::IndexIDMap* IDMapIndex = new faiss::IndexIDMap(PQIndex);
+	faiss::IndexIDMap2* IDMapIndex = new faiss::IndexIDMap2(PQIndex);
 	return TUniquePtr<FPQVectorIndex>(new FPQVectorIndex(ExpectedDimension, LoadedSubvectorSize, LoadedNBits, IDMapIndex));
 }
 
@@ -221,7 +324,9 @@ TUniquePtr<FPQVectorIndex> FPQVectorIndex::Deserialize(TConstArrayView<uint8> Da
 		return nullptr;
 	}
 
-	faiss::IndexIDMap* IDMapIndex = static_cast<faiss::IndexIDMap*>(LoadedIndex);
+	faiss::IndexIDMap2* IDMapIndex = static_cast<faiss::IndexIDMap2*>(LoadedIndex);
+	IDMapIndex->construct_rev_map();
+
 	faiss::IndexPQ* PQIndex = static_cast<faiss::IndexPQ*>(IDMapIndex->index);
 	const int32 LoadedM = static_cast<int32>(PQIndex->pq.M);
 	const int32 LoadedNBits = static_cast<int32>(PQIndex->pq.nbits);
diff --git a/Engine/Plugins/Experimental/SemanticSearch/Source/SemanticSearch/Private/VectorIndex/PQVectorIndex.h b/Engine/Plugins/Experimental/SemanticSearch/Source/SemanticSearch/Private/VectorIndex/PQVectorIndex.h
--- a/Engine/Plugins/Experimental/SemanticSearch/Source/SemanticSearch/Private/VectorIndex/PQVectorIndex.h
+++ b/Engine/Plugins/Experimental/SemanticSearch/Source/SemanticSearch/Private/VectorIndex/PQVectorIndex.h
@@ -2,7 +2,7 @@
 
 #pragma once
 
-#include "Interfaces/IVectorIndex.h"
+#include "Interfaces/IQuantizedVectorIndex.h"
 #include "HAL/CriticalSection.h"
 #include "Templates/UniquePtr.h"
 
@@ -26,7 +26,7 @@ namespace UE::SemanticSearch
  * @param InSubvectorSize  Number of dimensions per subquantizer. Dimension must be divisible by this.
  * @param InNBits          Bits per subquantizer code (8 = 256 centroids per subvector).
  */
-class FPQVectorIndex : public IVectorIndex
+class FPQVectorIndex : public IQuantizedVectorIndex
 {
 public:
 	FPQVectorIndex(int32 InDimension, int32 InSubvectorSize = 8, int32 InNBits = 8);
@@ -45,8 +45,11 @@ public:
 	virtual int32 GetDimension() const override { return Dimension; }
 	virtual TArray<uint8> Serialize() const override;
 	virtual TArray<uint8> Quantize(TConstArrayView<float> Vectors, int64 NumVectors) const override;
-	virtual bool SupportsQuantization() const override { return true; }
 	virtual void AddQuantized(TConstArrayView<int64> IDs, TConstArrayView<uint8> Codes) override;
+	virtual bool Contains(int64 ID) const override;
+	virtual bool TryGetEmbedding(int64 ID, TArray<float>& OutEmbedding) const override;
+	virtual bool Update(int64 OldID, int64 NewID) override;
+	virtual int64 EstimateMemoryBytes() const override;
 
 	/** Number of subquantizers (Dimension / SubvectorSize). */
 	int32 GetNumSubquantizers() const { return NumSubquantizers; }
@@ -57,8 +60,9 @@ public:
 	/** Bits per subquantizer code. */
 	int32 GetNBits() const { return NBits; }
 
-	/** Serialize just the PQ codebook (centroids, M, nbits). */
-	TArray<uint8> SerializeCodebook() const;
+	// IQuantizedVectorIndex overrides
+	virtual TArray<uint8> SerializeCodebook() const override;
+	virtual FIoHash GetCodebookHash() const override;
 
 	/** Create a trained empty index from a serialized codebook. */
 	static TUniquePtr<FPQVectorIndex> DeserializeCodebook(TConstArrayView<uint8> Data, int32 ExpectedDimension);
@@ -67,14 +71,16 @@ public:
 	static TUniquePtr<FPQVectorIndex> Deserialize(TConstArrayView<uint8> Data, int32 ExpectedDimension);
 
 private:
-	FPQVectorIndex(int32 InDimension, int32 InSubvectorSize, int32 InNBits, faiss::IndexIDMap* InIndex);
+	FPQVectorIndex(int32 InDimension, int32 InSubvectorSize, int32 InNBits, faiss::IndexIDMap2* InIndex);
 
 	int32 Dimension;
 	int32 SubvectorSize;
 	int32 NumSubquantizers;  // = Dimension / SubvectorSize
 	int32 NBits;
-	faiss::IndexIDMap* IndexWithIDs = nullptr;
+	faiss::IndexIDMap2* IndexWithIDs = nullptr;
 	mutable FRWLock IndexLock;
+	mutable FIoHash CachedCodebookHash;
+	mutable bool bCodebookHashCached = false;
 };
 
 } // namespace UE::SemanticSearch
diff --git a/Engine/Plugins/Experimental/SemanticSearch/Source/SemanticSearch/Private/VectorIndex/VectorIndexFactory.cpp b/Engine/Plugins/Experimental/SemanticSearch/Source/SemanticSearch/Private/VectorIndex/VectorIndexFactory.cpp
new file mode 100644
--- /dev/null
+++ b/Engine/Plugins/Experimental/SemanticSearch/Source/SemanticSearch/Private/VectorIndex/VectorIndexFactory.cpp
@@ -0,0 +1,78 @@
+// Copyright Epic Games, Inc. All Rights Reserved.
+
+#include "VectorIndexFactory.h"
+
+#include "FlatVectorIndex.h"
+#include "PQVectorIndex.h"
+#include "SemanticSearchModule.h"
+
+namespace UE::SemanticSearch
+{
+
+TSharedPtr<IVectorIndex> CreateVectorIndex(ESemanticSearchIndexType Type, int32 Dimension, const USemanticSearchSettings* Settings)
+{
+	if (Dimension <= 0)
+	{
+		return nullptr;
+	}
+	switch (Type)
+	{
+	case ESemanticSearchIndexType::PQ:
+	{
+		const int32 SubvectorSize = Settings ? Settings->PQSubvectorSize : 8;
+		const int32 NBits = Settings ? Settings->PQNBits : 8;
+		return MakeShared<FPQVectorIndex>(Dimension, SubvectorSize, NBits);
+	}
+
+	case ESemanticSearchIndexType::Flat:
+	default:
+		return MakeShared<FFlatVectorIndex>(Dimension);
+	}
+}
+
+TSharedPtr<IVectorIndex> DeserializeVectorIndex(ESemanticSearchIndexType Type, TConstArrayView<uint8> Data, int32 Dimension)
+{
+	switch (Type)
+	{
+	case ESemanticSearchIndexType::PQ:
+		return TSharedPtr<IVectorIndex>(FPQVectorIndex::Deserialize(Data, Dimension).Release());
+
+	case ESemanticSearchIndexType::Flat:
+		return TSharedPtr<IVectorIndex>(FFlatVectorIndex::Deserialize(Data, Dimension).Release());
+
+	default:
+		UE_LOG(LogSemanticSearch, Warning, TEXT("DeserializeVectorIndex: unsupported index type %d"), static_cast<uint8>(Type));
+		return nullptr;
+	}
+}
+
+TSharedPtr<IVectorIndex> DeserializeCodebook(ESemanticSearchIndexType Type, TConstArrayView<uint8> Data, int32 Dimension)
+{
+	switch (Type)
+	{
+	case ESemanticSearchIndexType::PQ:
+		return TSharedPtr<IVectorIndex>(FPQVectorIndex::DeserializeCodebook(Data, Dimension).Release());
+
+	case ESemanticSearchIndexType::Flat:
+		return nullptr;
+
+	default:
+		UE_LOG(LogSemanticSearch, Warning, TEXT("DeserializeCodebook: unsupported index type %d"), static_cast<uint8>(Type));
+		return nullptr;
+	}
+}
+
+FString GetCodebookFilename(ESemanticSearchIndexType Type)
+{
+	switch (Type)
+	{
+	case ESemanticSearchIndexType::PQ:
+		return TEXT("PQCodebook.bin");
+
+	case ESemanticSearchIndexType::Flat:
+	default:
+		return FString();
+	}
+}
+
+} // namespace UE::SemanticSearch
diff --git a/Engine/Plugins/Experimental/SemanticSearch/Source/SemanticSearch/Private/VectorIndex/VectorIndexFactory.h b/Engine/Plugins/Experimental/SemanticSearch/Source/SemanticSearch/Private/VectorIndex/VectorIndexFactory.h
new file mode 100644
--- /dev/null
+++ b/Engine/Plugins/Experimental/SemanticSearch/Source/SemanticSearch/Private/VectorIndex/VectorIndexFactory.h
@@ -0,0 +1,24 @@
+// Copyright Epic Games, Inc. All Rights Reserved.
+
+#pragma once
+
+#include "Interfaces/IVectorIndex.h"
+#include "Settings/SemanticSearchSettings.h"
+#include "Templates/SharedPointer.h"
+
+namespace UE::SemanticSearch
+{
+
+/** Create a fresh vector index for the given type. Quantized types start untrained. */
+TSharedPtr<IVectorIndex> CreateVectorIndex(ESemanticSearchIndexType Type, int32 Dimension, const USemanticSearchSettings* Settings);
+
+/** Deserialize a full index (with vectors) from bytes. Returns nullptr on failure. */
+TSharedPtr<IVectorIndex> DeserializeVectorIndex(ESemanticSearchIndexType Type, TConstArrayView<uint8> Data, int32 Dimension);
+
+/** Deserialize a trained-but-empty codebook. Returns nullptr for Flat or on failure. */
+TSharedPtr<IVectorIndex> DeserializeCodebook(ESemanticSearchIndexType Type, TConstArrayView<uint8> Data, int32 Dimension);
+
+/** Get the codebook filename for an index type (e.g. "PQCodebook.bin"). Returns empty for types without codebooks. */
+FString GetCodebookFilename(ESemanticSearchIndexType Type);
+
+} // namespace UE::SemanticSearch
diff --git a/Engine/Plugins/Experimental/SemanticSearch/Source/SemanticSearch/Public/AssetProcessorManager.h b/Engine/Plugins/Experimental/SemanticSearch/Source/SemanticSearch/Public/AssetProcessorManager.h
--- a/Engine/Plugins/Experimental/SemanticSearch/Source/SemanticSearch/Public/AssetProcessorManager.h
+++ b/Engine/Plugins/Experimental/SemanticSearch/Source/SemanticSearch/Public/AssetProcessorManager.h
@@ -10,6 +10,7 @@
 #include "DerivedDataCacheKey.h"
 #include "HAL/CriticalSection.h"
 #include "Interfaces/IAssetProcessor.h"
+#include "IO/IoHash.h"
 #include "Templates/Function.h"
 #include "Templates/UniquePtr.h"
 #include "Templates/UnrealTypeTraits.h"
@@ -34,6 +35,11 @@ struct FAssetEmbeddingResult
 	TArray<float> Embedding;
 };
 
+struct FAssetQuantizedEmbeddingResult
+{
+	TArray<uint8> QuantizedCodes;
+};
+
 class FAssetProcessorManager
 {
 public:
@@ -43,6 +49,8 @@ public:
 	SEMANTICSEARCH_API bool StarCachingAsset(FAssetData&& InAssetData);
 	SEMANTICSEARCH_API bool GetCaptionData(const FAssetData& InAssetData, TUniqueFunction<void(FAssetCaptionResult&&, FString&&)>&& InOnDataAvailable);
 	SEMANTICSEARCH_API bool GetEmbeddingData(const FAssetData& InAssetData, TUniqueFunction<void(FAssetEmbeddingResult&&, FString&&)>&& InOnDataAvailable, bool bBuildOnMiss = true);
+	SEMANTICSEARCH_API bool GetQuantizedEmbeddingData(const FAssetData& InAssetData, const FIoHash& CodebookHash,
+		TUniqueFunction<void(FAssetQuantizedEmbeddingResult&&, FString&&)>&& InOnDataAvailable, bool bBuildOnMiss = true);
 
 	/** Returns the processor for the asset's class. May return null. */
 	SEMANTICSEARCH_API TSharedPtr<IAssetProcessor> GetProcessorForAsset(const FAssetData& InAssetData) const;
diff --git a/Engine/Plugins/Experimental/SemanticSearch/Source/SemanticSearch/Public/HybridSearchIndex.h b/Engine/Plugins/Experimental/SemanticSearch/Source/SemanticSearch/Public/HybridSearchIndex.h
--- a/Engine/Plugins/Experimental/SemanticSearch/Source/SemanticSearch/Public/HybridSearchIndex.h
+++ b/Engine/Plugins/Experimental/SemanticSearch/Source/SemanticSearch/Public/HybridSearchIndex.h
@@ -8,9 +8,11 @@
 #include "Containers/StringView.h"
 #include "Interfaces/IVectorIndex.h"
 #include "IO/IoHash.h"
+#include "Settings/SemanticSearchSettings.h"
 #include "Templates/SharedPointer.h"
 #include "Templates/UniquePtr.h"
 #include "HAL/CriticalSection.h"
+#include "Delegates/Delegate.h"
 
 struct FAssetData;
 
@@ -27,6 +29,8 @@ struct FHybridSearchResult
 
 class FBM25Index;
 
+DECLARE_MULTICAST_DELEGATE_OneParam(FOnIndexChanged, bool /*bSuccess*/);
+
 /**
  * Singleton that owns all search indices (vector + BM25) and fuses results
  * using Reciprocal Rank Fusion (RRF).
@@ -39,6 +43,9 @@ public:
 	/** Get the vector index (may be null before first Add). */
 	SEMANTICSEARCH_API TSharedPtr<IVectorIndex> GetVectorIndex() const;
 
+	/** Replace the vector index. Used during index type switching. */
+	SEMANTICSEARCH_API void SetVectorIndex(TSharedPtr<IVectorIndex> NewIndex);
+
 	/** Create the vector index if not yet initialized. */
 	void EnsureInitialized(int32 EmbeddingDimension);
 
@@ -55,13 +62,16 @@ public:
 		int32 K,
 		TConstArrayView<int64> IDFilter = {}) const;
 
+	/** Add a single asset to both indices. Lazily creates vector index on first call. */
+	SEMANTICSEARCH_API void Add(const FAssetData& Asset,
+		TConstArrayView<float> Embedding, FStringView Caption, TConstArrayView<FString> Keywords);
+
 	/**
-	 * Add a single asset to both indices. Lazily creates vector index on first call.
-	 * @param QuantizedCodes  If non-empty and index is PQ, inserts pre-quantized codes directly (skips re-encoding).
+	 * Add a single asset using only pre-quantized codes (no float embedding).
+	 * The vector index must already be initialized and trained.
 	 */
-	SEMANTICSEARCH_API void Add(const FAssetData& Asset,
-		TConstArrayView<float> Embedding, FStringView Caption, TConstArrayView<FString> Keywords,
-		TConstArrayView<uint8> QuantizedCodes = {});
+	SEMANTICSEARCH_API void AddQuantized(const FAssetData& Asset,
+		TConstArrayView<uint8> QuantizedCodes, FStringView Caption, TConstArrayView<FString> Keywords);
 
 	/** Remove a single asset from both indices. */
 	SEMANTICSEARCH_API void Remove(const FAssetData& Asset);
@@ -109,6 +119,12 @@ public:
 	 */
 	SEMANTICSEARCH_API bool LoadFromFile(const FString& FilePath, TMap<int64, FIoHash>& OutPackageHashes);
 
+	/** Get the path to the shared codebook file for the given index type (in plugin Resources/). */
+	SEMANTICSEARCH_API static FString GetCodebookPath(ESemanticSearchIndexType Type);
+
+	/** Broadcast when the index changes (e.g. after switching type or retraining). */
+	FOnIndexChanged OnIndexChanged;
+
 private:
 	FHybridSearchIndex();
 
diff --git a/Engine/Plugins/Experimental/SemanticSearch/Source/SemanticSearch/Public/ISemanticSearchModule.h b/Engine/Plugins/Experimental/SemanticSearch/Source/SemanticSearch/Public/ISemanticSearchModule.h
--- a/Engine/Plugins/Experimental/SemanticSearch/Source/SemanticSearch/Public/ISemanticSearchModule.h
+++ b/Engine/Plugins/Experimental/SemanticSearch/Source/SemanticSearch/Public/ISemanticSearchModule.h
@@ -6,11 +6,25 @@
 #include "Interfaces/IEmbeddingProvider.h"
 #include "Modules/ModuleInterface.h"
 #include "Modules/ModuleManager.h"
+#include "Settings/SemanticSearchSettings.h"
 #include "Templates/SharedPointer.h"
 #include "Templates/UniquePtr.h"
 
 namespace UE::SemanticSearch
 {
+	struct FSemanticSearchIndexStats
+	{
+		ESemanticSearchIndexType IndexType = ESemanticSearchIndexType::Flat;
+		int32 VectorCount = 0;
+		int32 BM25Count = 0;
+		int32 Dimension = 0;
+		int64 EstimatedVectorMemoryBytes = 0;
+		int64 EstimatedBM25MemoryBytes = 0;
+		bool bIsTrained = false;
+		int32 SupportedAssetCount = 0;
+		int32 FailedCount = 0;
+	};
+
 	class ISemanticSearchModule : public IModuleInterface
 	{
 	public:
@@ -33,5 +47,13 @@ namespace UE::SemanticSearch
 
 		/** Returns true if the given package path is in a folder configured for indexing. */
 		SEMANTICSEARCH_API static bool IsInIndexedFolder(FStringView PackagePath);
+		/** Get current index statistics. */
+		virtual FSemanticSearchIndexStats GetIndexStats() const = 0;
+
+		/** Switch to a different index type. Re-indexes all assets. */
+		virtual void SwitchIndexType(ESemanticSearchIndexType NewType) = 0;
+
+		/** Retrain the current quantized index from scratch (new codebook). No-op for flat index. */
+		virtual void RetrainIndex() = 0;
 	};
-}
\ No newline at end of file
+}
diff --git a/Engine/Plugins/Experimental/SemanticSearch/Source/SemanticSearch/Public/Interfaces/IQuantizedVectorIndex.h b/Engine/Plugins/Experimental/SemanticSearch/Source/SemanticSearch/Public/Interfaces/IQuantizedVectorIndex.h
new file mode 100644
--- /dev/null
+++ b/Engine/Plugins/Experimental/SemanticSearch/Source/SemanticSearch/Public/Interfaces/IQuantizedVectorIndex.h
@@ -0,0 +1,52 @@
+// Copyright Epic Games, Inc. All Rights Reserved.
+
+#pragma once
+
+#include "Interfaces/IVectorIndex.h"
+#include "IO/IoHash.h"
+
+namespace UE::SemanticSearch
+{
+
+/**
+ * Interface for vector indices that support quantization (PQ, RaBitQ, etc.).
+ *
+ * Adds quantization, codebook serialization, and pre-quantized insertion
+ * to the base IVectorIndex interface.
+ *
+ * Callers should check IVectorIndex::SupportsQuantization() before casting
+ * to this type via static_cast.
+ */
+class IQuantizedVectorIndex : public IVectorIndex
+{
+public:
+	virtual bool SupportsQuantization() const override final { return true; }
+
+	/**
+	 * Quantize vectors into compressed form (e.g. PQ codes, RaBitQ codes).
+	 */
+	virtual TArray<uint8> Quantize(TConstArrayView<float> Vectors, int64 NumVectors) const = 0;
+
+	TArray<uint8> Quantize(TConstArrayView<float> Vector) const
+	{
+		return Quantize(Vector, 1);
+	}
+
+	/**
+	 * Add pre-quantized vectors directly, skipping encoding.
+	 */
+	virtual void AddQuantized(TConstArrayView<int64> IDs, TConstArrayView<uint8> Codes) = 0;
+
+	void AddQuantized(int64 ID, TConstArrayView<uint8> Codes)
+	{
+		AddQuantized(TConstArrayView<int64>(&ID, 1), Codes);
+	}
+
+	/** Serialize the trained codebook/quantizer state (no vector data). */
+	virtual TArray<uint8> SerializeCodebook() const = 0;
+
+	/** Blake3 hash of the serialized codebook. Used for DDC invalidation and save-file validation. */
+	virtual FIoHash GetCodebookHash() const = 0;
+};
+
+} // namespace UE::SemanticSearch
diff --git a/Engine/Plugins/Experimental/SemanticSearch/Source/SemanticSearch/Public/Interfaces/IVectorIndex.h b/Engine/Plugins/Experimental/SemanticSearch/Source/SemanticSearch/Public/Interfaces/IVectorIndex.h
--- a/Engine/Plugins/Experimental/SemanticSearch/Source/SemanticSearch/Public/Interfaces/IVectorIndex.h
+++ b/Engine/Plugins/Experimental/SemanticSearch/Source/SemanticSearch/Public/Interfaces/IVectorIndex.h
@@ -52,33 +52,9 @@ public:
 	/** Serialize the full index to a byte array. */
 	virtual TArray<uint8> Serialize() const = 0;
 
-	/** Whether this index supports quantization (Quantize/AddQuantized). */
+	/** Whether this index supports quantization. If true, can static_cast to IQuantizedVectorIndex. */
 	virtual bool SupportsQuantization() const { return false; }
 
-	/**
-	 * Quantize vectors into compressed form (e.g. PQ codes).
-	 * check(SupportsQuantization()) — crashes if not supported.
-	 */
-	virtual TArray<uint8> Quantize(TConstArrayView<float> Vectors, int64 NumVectors) const
-	{
-		checkf(false, TEXT("Quantize called on index that does not support quantization"));
-		return {};
-	}
-
-	TArray<uint8> Quantize(TConstArrayView<float> Vector) const
-	{
-		return Quantize(Vector, 1);
-	}
-
-	/**
-	 * Add pre-quantized vectors directly, skipping encoding.
-	 * check(SupportsQuantization()) — crashes if not supported.
-	 */
-	virtual void AddQuantized(TConstArrayView<int64> IDs, TConstArrayView<uint8> Codes)
-	{
-		checkf(false, TEXT("AddQuantized called on index that does not support quantization"));
-	}
-
 	// --- Inline single-item helpers ---
 
 	void Add(int64 ID, TConstArrayView<float> Vector)
@@ -86,11 +62,6 @@ public:
 		Add(TConstArrayView<int64>(&ID, 1), Vector);
 	}
 
-	void AddQuantized(int64 ID, TConstArrayView<uint8> Codes)
-	{
-		AddQuantized(TConstArrayView<int64>(&ID, 1), Codes);
-	}
-
 	bool Remove(int64 ID)
 	{
 		const int64 CountBefore = GetCount();
@@ -105,6 +76,9 @@ public:
 	/** Retrieve the embedding for a given ID. Returns false if not supported or ID not found. */
 	virtual bool TryGetEmbedding(int64 ID, TArray<float>& OutEmbedding) const { return false; }
 
+	/** Estimated memory usage in bytes. Returns 0 if not implemented. */
+	virtual int64 EstimateMemoryBytes() const { return 0; }
+
 	/**
 	 * Re-key a vector from OldID to NewID, preserving the embedding data.
 	 * Returns true if the vector was found and updated, false if OldID was not in the index.
diff --git a/Engine/Plugins/Experimental/SemanticSearch/Source/SemanticSearch/Public/Settings/SemanticSearchSettings.h b/Engine/Plugins/Experimental/SemanticSearch/Source/SemanticSearch/Public/Settings/SemanticSearchSettings.h
--- a/Engine/Plugins/Experimental/SemanticSearch/Source/SemanticSearch/Public/Settings/SemanticSearchSettings.h
+++ b/Engine/Plugins/Experimental/SemanticSearch/Source/SemanticSearch/Public/Settings/SemanticSearchSettings.h
@@ -9,7 +9,7 @@ UENUM()
 enum class ESemanticSearchIndexType : uint8
 {
 	Flat    UMETA(DisplayName = "Flat (Exact)"),
-	PQ      UMETA(DisplayName = "Product Quantization (Compressed)")
+	PQ      UMETA(DisplayName = "PQ (Compressed)")
 };
 
 UCLASS(config = EditorPerProjectUserSettings, DefaultConfig, meta = (DisplayName = "Semantic Search"))
@@ -39,9 +39,22 @@ public:
 		meta = (ToolTip = "When true, assets are automatically re-indexed when saved or updated on disk."))
 	bool bAutoIndexOnSave = true;
 
-	UPROPERTY(EditAnywhere, Config, Category = "Vector Index")
+	UPROPERTY(VisibleAnywhere, Config, Category = "Vector Index",
+		meta = (ToolTip = "Index type is managed via Tools > Semantic Search."))
 	ESemanticSearchIndexType IndexType = ESemanticSearchIndexType::Flat;
 
+	UPROPERTY(EditAnywhere, Config, Category = "Vector Index",
+		meta = (EditCondition = "IndexType == ESemanticSearchIndexType::PQ", EditConditionHides,
+			ClampMin = "1", ClampMax = "64",
+			ToolTip = "Number of dimensions per subquantizer. Embedding dimension must be divisible by this."))
+	int32 PQSubvectorSize = 8;
+
+	UPROPERTY(EditAnywhere, Config, Category = "Vector Index",
+		meta = (EditCondition = "IndexType == ESemanticSearchIndexType::PQ", EditConditionHides,
+			ClampMin = "4", ClampMax = "16",
+			ToolTip = "Bits per subquantizer code. 8 = 256 centroids per subvector."))
+	int32 PQNBits = 8;
+
 	UPROPERTY(EditAnywhere, Config, Category = "Search")
 	bool bEnableHybridSearch = true;
 
diff --git a/Engine/Plugins/Experimental/SemanticSearch/Source/SemanticSearch/SemanticSearch.Build.cs b/Engine/Plugins/Experimental/SemanticSearch/Source/SemanticSearch/SemanticSearch.Build.cs
--- a/Engine/Plugins/Experimental/SemanticSearch/Source/SemanticSearch/SemanticSearch.Build.cs
+++ b/Engine/Plugins/Experimental/SemanticSearch/Source/SemanticSearch/SemanticSearch.Build.cs
@@ -23,6 +23,7 @@ public class SemanticSearch : ModuleRules
 				"ImageCore",
 				"Json",
 				"PhysicsCore",
+				"Projects",
 				"Slate",
 				"UnrealEd",
 				"DeveloperSettings",
diff --git a/Engine/Plugins/Experimental/SemanticSearch/Source/SemanticSearchEditorIntegrations/Private/ContentBrowser/SemanticSearchCB.cpp b/Engine/Plugins/Experimental/SemanticSearch/Source/SemanticSearchEditorIntegrations/Private/ContentBrowser/SemanticSearchCB.cpp
--- a/Engine/Plugins/Experimental/SemanticSearch/Source/SemanticSearchEditorIntegrations/Private/ContentBrowser/SemanticSearchCB.cpp
+++ b/Engine/Plugins/Experimental/SemanticSearch/Source/SemanticSearchEditorIntegrations/Private/ContentBrowser/SemanticSearchCB.cpp
@@ -383,11 +383,8 @@ static FDelayedAutoRegisterHelper SemanticSearchStatsNavBarExtension(
 
 					auto GetUnindexedCount = []() -> int32
 					{
-						const UE::SemanticSearch::FHybridSearchIndex& Index = UE::SemanticSearch::FHybridSearchIndex::Get();
-						const int32 Indexed = Index.GetIndexedCount();
-						const int32 Failed = Index.GetFailedCount();
-						const int32 TotalSupported = UE::SemanticSearch::FAssetProcessorManager::Get().GetSupportedAssetCount();
-						return FMath::Max(0, TotalSupported - Indexed - Failed);
+						const UE::SemanticSearch::FSemanticSearchIndexStats Stats = UE::SemanticSearch::ISemanticSearchModule::Get().GetIndexStats();
+						return FMath::Max(0, Stats.SupportedAssetCount - Stats.VectorCount - Stats.FailedCount);
 					};
 
 					FToolMenuEntry& Entry = InSection.AddEntry(
@@ -411,11 +408,10 @@ static FDelayedAutoRegisterHelper SemanticSearchStatsNavBarExtension(
 									.TextStyle(FAppStyle::Get(), "SmallText")
 									.Text_Lambda([]() -> FText
 									{
-										const UE::SemanticSearch::FHybridSearchIndex& Index = UE::SemanticSearch::FHybridSearchIndex::Get();
-										const int32 Indexed = Index.GetIndexedCount();
-										const int32 Failed = Index.GetFailedCount();
-										const int32 TotalSupported = UE::SemanticSearch::FAssetProcessorManager::Get().GetSupportedAssetCount();
-										const int32 Unindexed = FMath::Max(0, TotalSupported - Indexed - Failed);
+										const UE::SemanticSearch::FSemanticSearchIndexStats Stats = UE::SemanticSearch::ISemanticSearchModule::Get().GetIndexStats();
+										const int32 Indexed = Stats.VectorCount;
+										const int32 Failed = Stats.FailedCount;
+										const int32 Unindexed = FMath::Max(0, Stats.SupportedAssetCount - Indexed - Failed);
 										if (Failed > 0)
 										{
 											return FText::Format(
diff --git a/Engine/Plugins/Experimental/SemanticSearch/Source/SemanticSearchEditorIntegrations/Private/SemanticSearchEditorIntegrationsModule.cpp b/Engine/Plugins/Experimental/SemanticSearch/Source/SemanticSearchEditorIntegrations/Private/SemanticSearchEditorIntegrationsModule.cpp
--- a/Engine/Plugins/Experimental/SemanticSearch/Source/SemanticSearchEditorIntegrations/Private/SemanticSearchEditorIntegrationsModule.cpp
+++ b/Engine/Plugins/Experimental/SemanticSearch/Source/SemanticSearchEditorIntegrations/Private/SemanticSearchEditorIntegrationsModule.cpp
@@ -5,6 +5,10 @@
 #include "ContentBrowserModule.h"
 #include "ContentBrowser/SemanticSearchCB.h"
 #include "Modules/ModuleManager.h"
+#include "ToolMenus.h"
+#include "Widgets/SSemanticSearchIndexDialog.h"
+
+#define LOCTEXT_NAMESPACE "SemanticSearchEditorIntegrations"
 
 namespace UE::SemanticSearch
 {
@@ -13,17 +17,46 @@ namespace UE::SemanticSearch
 	public:
 		virtual void StartupModule() override;
 		virtual void ShutdownModule() override;
+
+	private:
+		void RegisterToolsMenuEntry();
 	};
 
 	void FSemanticSearchEditorIntegrationsModule::StartupModule()
 	{
 		ContentBrowser::RegisterContentBrowserExtension();
+
+		UToolMenus::RegisterStartupCallback(
+			FSimpleMulticastDelegate::FDelegate::CreateRaw(this, &FSemanticSearchEditorIntegrationsModule::RegisterToolsMenuEntry));
 	}
 
 	void FSemanticSearchEditorIntegrationsModule::ShutdownModule()
 	{
 		ContentBrowser::UnregisterContentBrowserExtension();
 	}
+
+	void FSemanticSearchEditorIntegrationsModule::RegisterToolsMenuEntry()
+	{
+		UToolMenu* ToolsMenu = UToolMenus::Get()->ExtendMenu("LevelEditor.MainMenu.Tools");
+		if (!ToolsMenu)
+		{
+			return;
+		}
+
+		FToolMenuSection& Section = ToolsMenu->FindOrAddSection("SemanticSearch");
+		Section.AddMenuEntry(
+			"SemanticSearchIndex",
+			LOCTEXT("MenuLabel", "Semantic Search"),
+			LOCTEXT("MenuTooltip", "View and manage the semantic search index"),
+			FSlateIcon(),
+			FUIAction(FExecuteAction::CreateLambda([]()
+			{
+				SSemanticSearchIndexDialog::Open();
+			}))
+		);
+	}
 }
 
 IMPLEMENT_MODULE(UE::SemanticSearch::FSemanticSearchEditorIntegrationsModule, SemanticSearchEditorIntegrations)
+
+#undef LOCTEXT_NAMESPACE
diff --git a/Engine/Plugins/Experimental/SemanticSearch/Source/SemanticSearchEditorIntegrations/Private/Widgets/SSemanticSearchIndexDialog.cpp b/Engine/Plugins/Experimental/SemanticSearch/Source/SemanticSearchEditorIntegrations/Private/Widgets/SSemanticSearchIndexDialog.cpp
new file mode 100644
--- /dev/null
+++ b/Engine/Plugins/Experimental/SemanticSearch/Source/SemanticSearchEditorIntegrations/Private/Widgets/SSemanticSearchIndexDialog.cpp
@@ -0,0 +1,270 @@
+// Copyright Epic Games, Inc. All Rights Reserved.
+
+#include "Widgets/SSemanticSearchIndexDialog.h"
+
+#include "AssetProcessorManager.h"
+#include "HybridSearchIndex.h"
+#include "Framework/Application/SlateApplication.h"
+#include "Containers/Ticker.h"
+#include "Widgets/SWindow.h"
+#include "Settings/SemanticSearchSettings.h"
+#include "Widgets/Images/SThrobber.h"
+#include "Widgets/Input/SButton.h"
+#include "Widgets/Layout/SBox.h"
+#include "Widgets/Layout/SSeparator.h"
+#include "Widgets/Text/STextBlock.h"
+
+#define LOCTEXT_NAMESPACE "SemanticSearchIndexDialog"
+
+namespace UE::SemanticSearch
+{
+
+static FString FormatMemorySize(int64 Bytes)
+{
+	if (Bytes < 1024)
+	{
+		return FString::Printf(TEXT("%lld B"), Bytes);
+	}
+	if (Bytes < 1024 * 1024)
+	{
+		return FString::Printf(TEXT("%.1f KB"), Bytes / 1024.0);
+	}
+	if (Bytes < 1024LL * 1024 * 1024)
+	{
+		return FString::Printf(TEXT("%.1f MB"), Bytes / (1024.0 * 1024.0));
+	}
+	return FString::Printf(TEXT("%.2f GB"), Bytes / (1024.0 * 1024.0 * 1024.0));
+}
+
+void SSemanticSearchIndexDialog::Construct(const FArguments& InArgs)
+{
+	ChildSlot
+	[
+		SNew(SBox)
+		.MinDesiredWidth(500.0f)
+		.Padding(16.0f)
+		[
+			SNew(SVerticalBox)
+			+ SVerticalBox::Slot()
+			.AutoHeight()
+			[
+				SAssignNew(StatsBox, SVerticalBox)
+			]
+			+ SVerticalBox::Slot()
+			.AutoHeight()
+			[
+				SAssignNew(ThrobberBox, SVerticalBox)
+				.Visibility(EVisibility::Collapsed)
+			]
+		]
+	];
+
+	// Bind to index change notifications for auto-refresh
+	IndexChangedHandle = FHybridSearchIndex::Get().OnIndexChanged.AddSP(
+		this, &SSemanticSearchIndexDialog::OnIndexChanged);
+
+	RefreshStats();
+}
+
+SSemanticSearchIndexDialog::~SSemanticSearchIndexDialog()
+{
+	if (IndexChangedHandle.IsValid())
+	{
+		FHybridSearchIndex::Get().OnIndexChanged.Remove(IndexChangedHandle);
+	}
+}
+
+void SSemanticSearchIndexDialog::RefreshStats()
+{
+	CachedStats = ISemanticSearchModule::Get().GetIndexStats();
+
+	if (!StatsBox.IsValid())
+	{
+		return;
+	}
+
+	StatsBox->ClearChildren();
+
+	auto AddRow = [this](const FText& Label, const FText& Value)
+	{
+		StatsBox->AddSlot()
+		.AutoHeight()
+		.Padding(0, 2)
+		[
+			SNew(SHorizontalBox)
+			+ SHorizontalBox::Slot()
+			.FillWidth(0.55f)
+			[
+				SNew(STextBlock)
+				.Text(Label)
+				.Font(FCoreStyle::GetDefaultFontStyle("Bold", 10))
+			]
+			+ SHorizontalBox::Slot()
+			.FillWidth(0.45f)
+			[
+				SNew(STextBlock)
+				.Text(Value)
+			]
+		];
+	};
+
+	const UEnum* IndexTypeEnum = StaticEnum<ESemanticSearchIndexType>();
+	const FText IndexTypeDisplayName = IndexTypeEnum
+		? IndexTypeEnum->GetDisplayNameTextByValue(static_cast<int64>(CachedStats.IndexType))
+		: FText::FromString(TEXT("Unknown"));
+
+	AddRow(LOCTEXT("IndexType", "Index Type"), IndexTypeDisplayName);
+	AddRow(LOCTEXT("VectorCount", "Vector Index Count"), FText::AsNumber(CachedStats.VectorCount));
+	AddRow(LOCTEXT("BM25Count", "BM25 Index Count"), FText::AsNumber(CachedStats.BM25Count));
+	AddRow(LOCTEXT("Dimension", "Embedding Dimension"), FText::AsNumber(CachedStats.Dimension));
+	AddRow(LOCTEXT("VectorMemoryEst", "Est. Vector Index Memory Usage"), FText::FromString(FormatMemorySize(CachedStats.EstimatedVectorMemoryBytes)));
+	AddRow(LOCTEXT("BM25MemoryEst", "Est. BM25 Index Memory Usage"), FText::FromString(FormatMemorySize(CachedStats.EstimatedBM25MemoryBytes)));
+	AddRow(LOCTEXT("Trained", "Trained"), FText::FromString(CachedStats.bIsTrained ? TEXT("Yes") : TEXT("No")));
+	AddRow(LOCTEXT("SupportedAssets", "Supported Assets"), FText::AsNumber(CachedStats.SupportedAssetCount));
+	AddRow(LOCTEXT("FailedCount", "Failed To Index"), FText::AsNumber(CachedStats.FailedCount));
+
+	// Separator
+	StatsBox->AddSlot()
+	.AutoHeight()
+	.Padding(0, 8)
+	[
+		SNew(SSeparator)
+	];
+
+	// Add a switch button for each alternative index type
+	if (IndexTypeEnum)
+	{
+		for (int32 i = 0; i < IndexTypeEnum->NumEnums() - 1; ++i) // -1 to skip _MAX
+		{
+			const ESemanticSearchIndexType Type = static_cast<ESemanticSearchIndexType>(IndexTypeEnum->GetValueByIndex(i));
+			if (Type == CachedStats.IndexType)
+			{
+				continue;
+			}
+
+			const FText ButtonLabel = FText::Format(
+				LOCTEXT("SwitchToFmt", "Switch to {0}"),
+				IndexTypeEnum->GetDisplayNameTextByIndex(i));
+
+			StatsBox->AddSlot()
+			.AutoHeight()
+			.Padding(0, 4)
+			.HAlign(HAlign_Center)
+			[
+				SNew(SButton)
+				.Text(ButtonLabel)
+				.OnClicked_Lambda([this, Type]()
+				{
+					ISemanticSearchModule::Get().SwitchIndexType(Type);
+					ShowThrobber();
+					return FReply::Handled();
+				})
+			];
+		}
+	}
+
+	// Retrain button — only for quantized indices that are already trained
+	if (CachedStats.IndexType != ESemanticSearchIndexType::Flat && CachedStats.bIsTrained)
+	{
+		StatsBox->AddSlot()
+		.AutoHeight()
+		.Padding(0, 4)
+		.HAlign(HAlign_Center)
+		[
+			SNew(SButton)
+			.Text(LOCTEXT("RetrainIndex", "Retrain Index"))
+			.OnClicked_Lambda([this]()
+			{
+				ISemanticSearchModule::Get().RetrainIndex();
+				ShowThrobber();
+				return FReply::Handled();
+			})
+		];
+	}
+}
+
+void SSemanticSearchIndexDialog::ShowThrobber()
+{
+	if (!ThrobberBox.IsValid())
+	{
+		return;
+	}
+
+	ThrobberBox->ClearChildren();
+	ThrobberBox->AddSlot()
+	.AutoHeight()
+	.Padding(0, 8)
+	.HAlign(HAlign_Center)
+	[
+		SNew(SHorizontalBox)
+		+ SHorizontalBox::Slot()
+		.AutoWidth()
+		.VAlign(VAlign_Center)
+		[
+			SNew(SThrobber)
+		]
+		+ SHorizontalBox::Slot()
+		.AutoWidth()
+		.Padding(8, 0, 0, 0)
+		.VAlign(VAlign_Center)
+		[
+			SNew(STextBlock)
+			.Text(LOCTEXT("SwitchingIndex", "Switching index type..."))
+		]
+	];
+
+	ThrobberBox->SetVisibility(EVisibility::Visible);
+}
+
+void SSemanticSearchIndexDialog::OnIndexChanged(bool bSuccess)
+{
+	if (!bSuccess && ThrobberBox.IsValid())
+	{
+		// Replace throbber with error message
+		ThrobberBox->ClearChildren();
+		ThrobberBox->AddSlot()
+		.AutoHeight()
+		.Padding(0, 8)
+		.HAlign(HAlign_Center)
+		[
+			SNew(STextBlock)
+			.Text(LOCTEXT("SwitchFailed", "Switch failed. Check the log for details."))
+			.ColorAndOpacity(FLinearColor(1.f, 0.3f, 0.3f))
+		];
+
+		// Auto-hide after 5 seconds
+		TWeakPtr<SVerticalBox> WeakThrobberBox = ThrobberBox;
+		FTSTicker::GetCoreTicker().AddTicker(
+			FTickerDelegate::CreateLambda([WeakThrobberBox](float) -> bool
+			{
+				if (TSharedPtr<SVerticalBox> Box = WeakThrobberBox.Pin())
+				{
+					Box->SetVisibility(EVisibility::Collapsed);
+				}
+				return false; // one-shot
+			}), 10.0f);
+	}
+	else if (ThrobberBox.IsValid())
+	{
+		ThrobberBox->SetVisibility(EVisibility::Collapsed);
+	}
+	RefreshStats();
+}
+
+void SSemanticSearchIndexDialog::Open()
+{
+	TSharedRef<SWindow> Window = SNew(SWindow)
+		.Title(LOCTEXT("WindowTitle", "Semantic Search"))
+		.SizingRule(ESizingRule::Autosized)
+		.SupportsMaximize(false)
+		.SupportsMinimize(false)
+		[
+			SNew(SSemanticSearchIndexDialog)
+		];
+
+	FSlateApplication::Get().AddWindow(Window);
+}
+
+} // namespace UE::SemanticSearch
+
+#undef LOCTEXT_NAMESPACE
diff --git a/Engine/Plugins/Experimental/SemanticSearch/Source/SemanticSearchEditorIntegrations/Private/Widgets/SSemanticSearchIndexDialog.h b/Engine/Plugins/Experimental/SemanticSearch/Source/SemanticSearchEditorIntegrations/Private/Widgets/SSemanticSearchIndexDialog.h
new file mode 100644
--- /dev/null
+++ b/Engine/Plugins/Experimental/SemanticSearch/Source/SemanticSearchEditorIntegrations/Private/Widgets/SSemanticSearchIndexDialog.h
@@ -0,0 +1,38 @@
+// Copyright Epic Games, Inc. All Rights Reserved.
+
+#pragma once
+
+#include "ISemanticSearchModule.h"
+#include "Widgets/SCompoundWidget.h"
+#include "Widgets/SBoxPanel.h"
+
+class STextBlock;
+
+namespace UE::SemanticSearch
+{
+
+class SSemanticSearchIndexDialog : public SCompoundWidget
+{
+public:
+	SLATE_BEGIN_ARGS(SSemanticSearchIndexDialog) {}
+	SLATE_END_ARGS()
+
+	void Construct(const FArguments& InArgs);
+	~SSemanticSearchIndexDialog();
+
+	/** Open the dialog as a standalone window. */
+	static void Open();
+
+private:
+	void RefreshStats();
+	void ShowThrobber();
+	void OnIndexChanged(bool bSuccess);
+
+	FSemanticSearchIndexStats CachedStats;
+	TSharedPtr<SVerticalBox> StatsBox;
+	TSharedPtr<SVerticalBox> ThrobberBox;
+	TSharedPtr<STextBlock> StatusText;
+	FDelegateHandle IndexChangedHandle;
+};
+
+} // namespace UE::SemanticSearch
