[Semantic Seach] Only allow modules to be packaged in editor builds.

#Synced-CL 52565678

diff --git a/Engine/Plugins/Experimental/SemanticSearch/SemanticSearch.uplugin b/Engine/Plugins/Experimental/SemanticSearch/SemanticSearch.uplugin
--- a/Engine/Plugins/Experimental/SemanticSearch/SemanticSearch.uplugin
+++ b/Engine/Plugins/Experimental/SemanticSearch/SemanticSearch.uplugin
@@ -19,12 +19,14 @@
 		{
 			"Name": "SemanticSearch",
 			"Type": "Editor",
-			"LoadingPhase": "Default"
+			"LoadingPhase": "Default",
+			"TargetAllowList": ["Editor"]
 		},
 		{
 			"Name": "SemanticSearchEditorIntegrations",
 			"Type": "Editor",
-			"LoadingPhase": "Default"
+			"LoadingPhase": "Default",
+			"TargetAllowList": ["Editor"]
 		}
 	]
 }
\ No newline at end of file
