Prevent FAISS index from being saved multiple times OnFilesLoaded

#Synced-CL 52617675

diff --git a/Engine/Plugins/Experimental/SemanticSearch/Source/SemanticSearch/Private/SemanticSearchModule.cpp b/Engine/Plugins/Experimental/SemanticSearch/Source/SemanticSearch/Private/SemanticSearchModule.cpp
--- a/Engine/Plugins/Experimental/SemanticSearch/Source/SemanticSearch/Private/SemanticSearchModule.cpp
+++ b/Engine/Plugins/Experimental/SemanticSearch/Source/SemanticSearch/Private/SemanticSearchModule.cpp
@@ -483,6 +483,13 @@ void FSemanticSearchModule::DismissProgressNotification()
 
 void FSemanticSearchModule::OnFilesLoaded()
 {
+	// Guard against OnFilesLoaded being called multiple times by the asset registry
+	if (bFilesLoaded)
+	{
+		return;
+	}
+	bFilesLoaded = true;
+
 	// Anything left in LoadedPackageHashes was not seen during OnAssetAdded — the asset was deleted
 	if (!LoadedPackageHashes.IsEmpty())
 	{
diff --git a/Engine/Plugins/Experimental/SemanticSearch/Source/SemanticSearch/Private/SemanticSearchModule.h b/Engine/Plugins/Experimental/SemanticSearch/Source/SemanticSearch/Private/SemanticSearchModule.h
--- a/Engine/Plugins/Experimental/SemanticSearch/Source/SemanticSearch/Private/SemanticSearchModule.h
+++ b/Engine/Plugins/Experimental/SemanticSearch/Source/SemanticSearch/Private/SemanticSearchModule.h
@@ -80,5 +80,6 @@ private:
 	/** Package hashes loaded from saved index, used for staleness checks during startup. */
 	TMap<int64, FIoHash> LoadedPackageHashes;
 	TMap<FName, FTimerHandle> PendingPackageSavedTimers;
+	bool bFilesLoaded = false;
 };
 }
