? [SemanticSearch] Improve indexing UX and fix progress tracking

- Add reindex all, index unindexed, and retry failed actions as inline
refresh icons in content browser stats bar
- Disable action buttons while indexing is in progress
- Epoch-based invalidation prevents stale callbacks from corrupting counters
when switching index type mid-indexing
- Re-index from scratch when IndexedPaths setting changes
- Invalidate supported asset count cache when IndexedPaths changes
- Clear BM25 index and failed set when re-indexing all assets
- User-initiated actions always build on DDC miss
- Fix bIsSwitchingIndexType guard never being set
- Guard OnFilesLoaded against duplicate calls
#Synced-CL 52622993

diff --git a/Engine/Plugins/Experimental/SemanticSearch/Source/SemanticSearch/Private/AssetProcessorManager.cpp b/Engine/Plugins/Experimental/SemanticSearch/Source/SemanticSearch/Private/AssetProcessorManager.cpp
--- a/Engine/Plugins/Experimental/SemanticSearch/Source/SemanticSearch/Private/AssetProcessorManager.cpp
+++ b/Engine/Plugins/Experimental/SemanticSearch/Source/SemanticSearch/Private/AssetProcessorManager.cpp
@@ -32,7 +32,7 @@ namespace UE::SemanticSearch
 			OnAssetRemovedHandle = AR.OnAssetRemoved().AddRaw(this, &FAssetProcessorManager::InvalidateSupportedAssetCount);
 		}
 	}
-
+	
 	FAssetProcessorManager::~FAssetProcessorManager()
 	{
 		FCoreUObjectDelegates::OnObjectsReinstanced.Remove(OnObjectsReinstancedHandle);
@@ -71,6 +71,16 @@ namespace UE::SemanticSearch
 		CachedSupportedAssetCount = -1;
 	}
 
+	void FAssetProcessorManager::InvalidateSupportedAssetCount()
+	{
+		CachedSupportedAssetCount = -1;
+	}
+
+	void FAssetProcessorManager::ClearBuildCache()
+	{
+		BuildsJobs.Empty();
+	}
+
 	int32 FAssetProcessorManager::ComputeSupportedAssetCount() const
 	{
 		// Collect all unique class paths that have a processor entry.
diff --git a/Engine/Plugins/Experimental/SemanticSearch/Source/SemanticSearch/Private/HybridSearchIndex.cpp b/Engine/Plugins/Experimental/SemanticSearch/Source/SemanticSearch/Private/HybridSearchIndex.cpp
--- a/Engine/Plugins/Experimental/SemanticSearch/Source/SemanticSearch/Private/HybridSearchIndex.cpp
+++ b/Engine/Plugins/Experimental/SemanticSearch/Source/SemanticSearch/Private/HybridSearchIndex.cpp
@@ -107,6 +107,13 @@ FBM25Index& FHybridSearchIndex::GetBM25Index()
 	return *BM25Index;
 }
 
+void FHybridSearchIndex::ResetBM25AndFailedState()
+{
+	BM25Index->Clear();
+	FScopeLock Lock(&FailedLock);
+	FailedAssets.Empty();
+}
+
 int32 FHybridSearchIndex::GetIndexedCount() const
 {
 	return VectorIndex ? static_cast<int32>(VectorIndex->GetCount()) : 0;
@@ -124,6 +131,12 @@ int32 FHybridSearchIndex::GetFailedCount() const
 	return FailedAssets.Num();
 }
 
+bool FHybridSearchIndex::IsFailed(int64 AssetID) const
+{
+	FScopeLock Lock(&FailedLock);
+	return FailedAssets.Contains(AssetID);
+}
+
 int64 GetAssetIndexID(FStringView AssetPath)
 {
 	FBlake3 Hasher;
diff --git a/Engine/Plugins/Experimental/SemanticSearch/Source/SemanticSearch/Private/SemanticSearchModule.cpp b/Engine/Plugins/Experimental/SemanticSearch/Source/SemanticSearch/Private/SemanticSearchModule.cpp
--- a/Engine/Plugins/Experimental/SemanticSearch/Source/SemanticSearch/Private/SemanticSearchModule.cpp
+++ b/Engine/Plugins/Experimental/SemanticSearch/Source/SemanticSearch/Private/SemanticSearchModule.cpp
@@ -25,6 +25,17 @@
 
 DEFINE_LOG_CATEGORY(LogSemanticSearch);
 
+void USemanticSearchSettings::PostEditChangeProperty(FPropertyChangedEvent& PropertyChangedEvent)
+{
+	Super::PostEditChangeProperty(PropertyChangedEvent);
+	if (PropertyChangedEvent.MemberProperty
+		&& PropertyChangedEvent.MemberProperty->GetFName() == GET_MEMBER_NAME_CHECKED(USemanticSearchSettings, IndexedPaths))
+	{
+		UE::SemanticSearch::FAssetProcessorManager::Get().InvalidateSupportedAssetCount();
+		UE::SemanticSearch::Private::FSemanticSearchModule::Get().IndexAllAssets(/*bForceBuild=*/true);
+	}
+}
+
 bool UE::SemanticSearch::ISemanticSearchModule::IsInIndexedFolder(FStringView PackagePath)
 {
 	const USemanticSearchSettings* Settings = USemanticSearchSettings::Get();
@@ -135,6 +146,8 @@ void FSemanticSearchModule::TriggerIndexing(const FAssetData& AssetData, bool bB
 	}
 	OnIndexingStarted();
 
+	const uint32 CapturedEpoch = IndexingEpoch;
+
 	// Check up front if we're on a quantized index
 	FHybridSearchIndex& Index = FHybridSearchIndex::Get();
 	TSharedPtr<IVectorIndex> VectorIndex = Index.GetVectorIndex();
@@ -142,7 +155,7 @@ void FSemanticSearchModule::TriggerIndexing(const FAssetData& AssetData, bool bB
 
 	// Step 1: Fetch caption (always needed for both flat and quantized paths)
 	bool bStarted = FAssetProcessorManager::Get().GetCaptionData(AssetData,
-		[Asset = FAssetData(AssetData), bBuildOnMiss, bQuantized](FAssetCaptionResult&& CaptionResult, FString&& CaptionError) mutable
+		[Asset = FAssetData(AssetData), bBuildOnMiss, bQuantized, CapturedEpoch](FAssetCaptionResult&& CaptionResult, FString&& CaptionError) mutable
 		{
 			if (bQuantized)
 			{
@@ -155,7 +168,7 @@ void FSemanticSearchModule::TriggerIndexing(const FAssetData& AssetData, bool bB
 					FIoHash CodebookHash = QuantizedIndex->GetCodebookHash();
 
 					bool bStarted = FAssetProcessorManager::Get().GetQuantizedEmbeddingData(Asset, CodebookHash,
-						[Asset, CaptionResult = MoveTemp(CaptionResult)](
+						[Asset, CaptionResult = MoveTemp(CaptionResult), CapturedEpoch](
 							FAssetQuantizedEmbeddingResult&& QuantizedResult, FString&& QuantizedError) mutable
 						{
 							if (QuantizedResult.QuantizedCodes.Num() > 0)
@@ -167,19 +180,19 @@ void FSemanticSearchModule::TriggerIndexing(const FAssetData& AssetData, bool bB
 								UE_LOGF(LogSemanticSearch, Log, "Indexed %ls (quantized, index size: %lld)",
 									*Asset.GetObjectPathString(),
 									Index.GetVectorIndex() ? Index.GetVectorIndex()->GetCount() : 0);
-								FSemanticSearchModule::Get().OnIndexingFinished(true);
+								FSemanticSearchModule::Get().OnIndexingFinished(true, CapturedEpoch);
 							}
 							else
 							{
 								FHybridSearchIndex::Get().MarkFailed(GetAssetIndexID(Asset));
 								UE_LOGF(LogSemanticSearch, Warning, "Failed to get quantized codes for %ls: %ls",
 									*Asset.GetObjectPathString(), *QuantizedError);
-								FSemanticSearchModule::Get().OnIndexingFinished(false);
+								FSemanticSearchModule::Get().OnIndexingFinished(false, CapturedEpoch);
 							}
 						}, bBuildOnMiss);
 					if (!bStarted)
 					{
-						FSemanticSearchModule::Get().OnIndexingFinished(false);
+						FSemanticSearchModule::Get().OnIndexingFinished(false, CapturedEpoch);
 					}
 					return;
 				}
@@ -187,7 +200,7 @@ void FSemanticSearchModule::TriggerIndexing(const FAssetData& AssetData, bool bB
 
 			// Step 2b (flat or index changed): fetch embedding and add
 			bool bStarted = FAssetProcessorManager::Get().GetEmbeddingData(Asset,
-				[Asset, CaptionResult = MoveTemp(CaptionResult), bBuildOnMiss](
+				[Asset, CaptionResult = MoveTemp(CaptionResult), bBuildOnMiss, CapturedEpoch](
 					FAssetEmbeddingResult&& EmbeddingResult, FString&& Error) mutable
 				{
 					if (!Error.IsEmpty() || EmbeddingResult.Embedding.IsEmpty())
@@ -199,7 +212,7 @@ void FSemanticSearchModule::TriggerIndexing(const FAssetData& AssetData, bool bB
 								*Asset.GetObjectPathString(),
 								Error.IsEmpty() ? TEXT("empty embedding") : *Error);
 						}
-						FSemanticSearchModule::Get().OnIndexingFinished(false);
+						FSemanticSearchModule::Get().OnIndexingFinished(false, CapturedEpoch);
 						return;
 					}
 
@@ -210,24 +223,24 @@ void FSemanticSearchModule::TriggerIndexing(const FAssetData& AssetData, bool bB
 					UE_LOGF(LogSemanticSearch, Log, "Indexed %ls (%d dims, index size: %lld)",
 						*Asset.GetObjectPathString(), EmbeddingResult.Embedding.Num(),
 						Index.GetVectorIndex() ? Index.GetVectorIndex()->GetCount() : 0);
-					FSemanticSearchModule::Get().OnIndexingFinished(true);
+					FSemanticSearchModule::Get().OnIndexingFinished(true, CapturedEpoch);
 				}, bBuildOnMiss);
 			if (!bStarted)
 			{
-				FSemanticSearchModule::Get().OnIndexingFinished(false);
+				FSemanticSearchModule::Get().OnIndexingFinished(false, CapturedEpoch);
 			}
 		});
 
 	if (!bStarted)
 	{
-		FSemanticSearchModule::Get().OnIndexingFinished(false);
+		FSemanticSearchModule::Get().OnIndexingFinished(false, CapturedEpoch);
 	}
 }
 
 void FSemanticSearchModule::OnAssetAdded(const FAssetData& AssetData)
 {
 	if (AssetData.IsRedirector() || !IsInIndexedFolder(AssetData)) { return; }
-
+	
 	FHybridSearchIndex& Index = FHybridSearchIndex::Get();
 	if (!LoadedPackageHashes.IsEmpty() && Index.Contains(AssetData))
 	{
@@ -243,7 +256,7 @@ void FSemanticSearchModule::OnAssetAdded(const FAssetData& AssetData)
 			}
 		}
 	}
-
+	
 	const USemanticSearchSettings* Settings = USemanticSearchSettings::Get();
 	UE_LOGF(LogSemanticSearch, Log, "OnAssetAdded: %ls", *AssetData.GetObjectPathString());
 	TriggerIndexing(AssetData, Settings && Settings->bBuildUncachedOnStartup);
@@ -413,10 +426,15 @@ void FSemanticSearchModule::OnIndexingStarted()
 	}
 }
 
-void FSemanticSearchModule::OnIndexingFinished(bool bSuccess)
+void FSemanticSearchModule::OnIndexingFinished(bool bSuccess, uint32 Epoch)
 {
-	AsyncTask(ENamedThreads::GameThread, [this, bSuccess]()
+	AsyncTask(ENamedThreads::GameThread, [this, bSuccess, Epoch]()
 	{
+		if (Epoch != IndexingEpoch)
+		{
+			return;
+		}
+
 		if (bSuccess)
 		{
 			++IndexingCompleted;
@@ -428,6 +446,11 @@ void FSemanticSearchModule::OnIndexingFinished(bool bSuccess)
 
 		UpdateProgressNotification();
 
+		if ((IndexingCompleted + IndexingFailed) % 50 == 0)
+		{
+			FHybridSearchIndex::Get().OnIndexChanged.Broadcast(true);
+		}
+
 		if (IndexingCompleted + IndexingFailed >= IndexingTotal && GEditor)
 		{
 			GEditor->GetTimerManager()->SetTimer(DismissTimerHandle,
@@ -481,6 +504,26 @@ void FSemanticSearchModule::DismissProgressNotification()
 	FHybridSearchIndex::Get().OnIndexChanged.Broadcast(true);
 }
 
+void FSemanticSearchModule::ResetProgressState()
+{
+	++IndexingEpoch;
+
+	if (GEditor)
+	{
+		GEditor->GetTimerManager()->ClearTimer(DismissTimerHandle);
+	}
+
+	if (ProgressHandle.IsValid())
+	{
+		FSlateNotificationManager::Get().CancelProgressNotification(ProgressHandle);
+		ProgressHandle = FProgressNotificationHandle();
+	}
+
+	IndexingTotal = 0;
+	IndexingCompleted = 0;
+	IndexingFailed = 0;
+}
+
 void FSemanticSearchModule::OnFilesLoaded()
 {
 	// Guard against OnFilesLoaded being called multiple times by the asset registry
@@ -534,6 +577,7 @@ FSemanticSearchIndexStats FSemanticSearchModule::GetIndexStats() const
 	Stats.EstimatedBM25MemoryBytes = Index.GetBM25Index().EstimateMemoryBytes();
 	Stats.FailedCount = Index.GetFailedCount();
 	Stats.SupportedAssetCount = FAssetProcessorManager::Get().GetSupportedAssetCount();
+	Stats.bIsIndexing = ProgressHandle.IsValid() && IndexingTotal > 0 && (IndexingCompleted + IndexingFailed) < IndexingTotal;
 
 	return Stats;
 }
@@ -552,6 +596,8 @@ void FSemanticSearchModule::SwitchIndexType(ESemanticSearchIndexType NewType)
 		return;
 	}
 
+	bIsSwitchingIndexType = true;
+
 	if (NewType == ESemanticSearchIndexType::Flat)
 	{
 		SwitchToFlatIndex();
@@ -580,7 +626,7 @@ void FSemanticSearchModule::SwitchToFlatIndex()
 	bIsSwitchingIndexType = false;
 
 	// Re-index all assets — TriggerIndexing fetches embeddings from DDC
-	IndexAllAssets();
+	IndexAllAssets(/*bForceBuild=*/true);
 }
 
 void FSemanticSearchModule::SwitchToQuantizedIndex(ESemanticSearchIndexType NewType, bool bForceTrain)
@@ -627,7 +673,7 @@ void FSemanticSearchModule::SwitchToQuantizedIndex(ESemanticSearchIndexType NewT
 		bIsSwitchingIndexType = false;
 
 		// TriggerIndexing will use the quantized fast path (DDC codes) when available
-		IndexAllAssets();
+		IndexAllAssets(/*bForceBuild=*/true);
 	}
 	else
 	{
@@ -727,19 +773,56 @@ void FSemanticSearchModule::SwitchToQuantizedIndex(ESemanticSearchIndexType NewT
 
 						UE_LOGF(LogSemanticSearch, Log, "Quantized index trained, re-indexing %d assets from DDC...", NumVectors);
 
-						IndexAllAssets();
+						IndexAllAssets(/*bForceBuild=*/true);
 					});
 			});
 	}
 }
 
-void FSemanticSearchModule::IndexAllAssets()
+void FSemanticSearchModule::IndexAllAssets(bool bForceBuild)
 {
+	ResetProgressState();
+	FHybridSearchIndex::Get().ResetBM25AndFailedState();
+
+	if (bForceBuild)
+	{
+		// Clear cached build tasks so fresh DDC lookups are triggered
+		FAssetProcessorManager::Get().ClearBuildCache();
+
+		// Clear and recreate the vector index so count drops to 0 and rebuilds
+		FHybridSearchIndex& Index = FHybridSearchIndex::Get();
+		TSharedPtr<IVectorIndex> CurrentIndex = Index.GetVectorIndex();
+		if (CurrentIndex)
+		{
+			const int32 Dimension = CurrentIndex->GetDimension();
+			const USemanticSearchSettings* Settings = USemanticSearchSettings::Get();
+			const ESemanticSearchIndexType Type = Settings ? Settings->IndexType : ESemanticSearchIndexType::Flat;
+
+			// Try to reload trained codebook for PQ types
+			TSharedPtr<IVectorIndex> NewIndex;
+			const FString CodebookPath = FHybridSearchIndex::GetCodebookPath(Type);
+			if (!CodebookPath.IsEmpty())
+			{
+				TArray<uint8> CodebookData;
+				if (FFileHelper::LoadFileToArray(CodebookData, *CodebookPath))
+				{
+					NewIndex = DeserializeCodebook(Type, CodebookData, Dimension);
+				}
+			}
+			if (!NewIndex)
+			{
+				NewIndex = CreateVectorIndex(Type, Dimension, Settings);
+			}
+			Index.SetVectorIndex(NewIndex);
+		}
+	}
+
 	IAssetRegistry& AssetRegistry = FModuleManager::LoadModuleChecked<FAssetRegistryModule>(TEXT("AssetRegistry")).Get();
 	TArray<FAssetData> AllAssets;
 	AssetRegistry.GetAllAssets(AllAssets, /*bIncludeOnlyOnDiskAssets=*/true);
 
 	const USemanticSearchSettings* Settings = USemanticSearchSettings::Get();
+	const bool bBuildOnMiss = bForceBuild || (Settings && Settings->bBuildUncachedOnStartup);
 
 	for (const FAssetData& Asset : AllAssets)
 	{
@@ -747,7 +830,7 @@ void FSemanticSearchModule::IndexAllAssets()
 		{
 			continue;
 		}
-		TriggerIndexing(Asset, Settings && Settings->bBuildUncachedOnStartup);
+		TriggerIndexing(Asset, bBuildOnMiss);
 	}
 }
 
@@ -766,6 +849,7 @@ void FSemanticSearchModule::RetrainIndex()
 		return;
 	}
 
+	bIsSwitchingIndexType = true;
 	SwitchToQuantizedIndex(Settings->IndexType, /*bForceTrain=*/true);
 }
 
diff --git a/Engine/Plugins/Experimental/SemanticSearch/Source/SemanticSearch/Private/SemanticSearchModule.h b/Engine/Plugins/Experimental/SemanticSearch/Source/SemanticSearch/Private/SemanticSearchModule.h
--- a/Engine/Plugins/Experimental/SemanticSearch/Source/SemanticSearch/Private/SemanticSearchModule.h
+++ b/Engine/Plugins/Experimental/SemanticSearch/Source/SemanticSearch/Private/SemanticSearchModule.h
@@ -33,13 +33,13 @@ public:
 	virtual void IndexAsset(const FAssetData& AssetData) override;
 	virtual FSemanticSearchIndexStats GetIndexStats() const override;
 	virtual void SwitchIndexType(ESemanticSearchIndexType NewType) override;
+	virtual void IndexAllAssets(bool bForceBuild = false) override;
 	virtual void RetrainIndex() override;
 
 private:
 	void SwitchToFlatIndex();
 	void SwitchToQuantizedIndex(ESemanticSearchIndexType NewType, bool bForceTrain = false);
 	void TriggerIndexing(const FAssetData& AssetData, bool bBuildOnMiss);
-	void IndexAllAssets();
 
 	void OnAssetAdded(const FAssetData& AssetData);
 	void OnPackageSaved(const FString& PackageFileName, UPackage* Package, FObjectPostSaveContext ObjectSaveContext);
@@ -49,9 +49,10 @@ private:
 	void OnFilesLoaded();
 
 	void OnIndexingStarted();
-	void OnIndexingFinished(bool bSuccess);
+	void OnIndexingFinished(bool bSuccess, uint32 Epoch);
 	void UpdateProgressNotification();
 	void DismissProgressNotification();
+	void ResetProgressState();
 	void SaveIndex();
 
 	static FString GetSavedIndexPath();
@@ -68,6 +69,7 @@ private:
 	int32 IndexingTotal = 0;
 	int32 IndexingCompleted = 0;
 	int32 IndexingFailed = 0;
+	uint32 IndexingEpoch = 0;
 	FTimerHandle DismissTimerHandle;
 	FTimerHandle PeriodicSaveTimerHandle;
 
diff --git a/Engine/Plugins/Experimental/SemanticSearch/Source/SemanticSearch/Private/TextIndex/BM25Index.cpp b/Engine/Plugins/Experimental/SemanticSearch/Source/SemanticSearch/Private/TextIndex/BM25Index.cpp
--- a/Engine/Plugins/Experimental/SemanticSearch/Source/SemanticSearch/Private/TextIndex/BM25Index.cpp
+++ b/Engine/Plugins/Experimental/SemanticSearch/Source/SemanticSearch/Private/TextIndex/BM25Index.cpp
@@ -408,6 +408,15 @@ int64 FBM25Index::EstimateMemoryBytes() const
 	return Bytes;
 }
 
+void FBM25Index::Clear()
+{
+	FWriteScopeLock WriteLock(IndexLock);
+	PostingLists.Empty();
+	DocLengths.Empty();
+	DocTokens.Empty();
+	TotalTokenCount = 0;
+}
+
 bool FBM25Index::Contains(int64 ID) const
 {
 	FReadScopeLock ReadLock(IndexLock);
diff --git a/Engine/Plugins/Experimental/SemanticSearch/Source/SemanticSearch/Private/TextIndex/BM25Index.h b/Engine/Plugins/Experimental/SemanticSearch/Source/SemanticSearch/Private/TextIndex/BM25Index.h
--- a/Engine/Plugins/Experimental/SemanticSearch/Source/SemanticSearch/Private/TextIndex/BM25Index.h
+++ b/Engine/Plugins/Experimental/SemanticSearch/Source/SemanticSearch/Private/TextIndex/BM25Index.h
@@ -54,6 +54,9 @@ public:
 	/** Estimate the memory footprint of this index in bytes. */
 	int64 EstimateMemoryBytes() const;
 
+	/** Clear all documents from the index. */
+	void Clear();
+
 private:
 	struct FPostingEntry
 	{
diff --git a/Engine/Plugins/Experimental/SemanticSearch/Source/SemanticSearch/Public/AssetProcessorManager.h b/Engine/Plugins/Experimental/SemanticSearch/Source/SemanticSearch/Public/AssetProcessorManager.h
--- a/Engine/Plugins/Experimental/SemanticSearch/Source/SemanticSearch/Public/AssetProcessorManager.h
+++ b/Engine/Plugins/Experimental/SemanticSearch/Source/SemanticSearch/Public/AssetProcessorManager.h
@@ -58,6 +58,12 @@ public:
 	/** Returns the total number of supported assets in indexed folders. Cached; invalidated on asset registry changes. */
 	SEMANTICSEARCH_API int32 GetSupportedAssetCount() const;
 
+	/** Force-invalidate the supported asset count cache (e.g. when IndexedPaths changes). */
+	SEMANTICSEARCH_API void InvalidateSupportedAssetCount();
+
+	/** Clear all cached build tasks. Necessary before re-indexing to ensure fresh DDC lookups. */
+	SEMANTICSEARCH_API void ClearBuildCache();
+
 	template<typename FAssetProcessor>
 	void RegisterAssetProcessor();
 
diff --git a/Engine/Plugins/Experimental/SemanticSearch/Source/SemanticSearch/Public/HybridSearchIndex.h b/Engine/Plugins/Experimental/SemanticSearch/Source/SemanticSearch/Public/HybridSearchIndex.h
--- a/Engine/Plugins/Experimental/SemanticSearch/Source/SemanticSearch/Public/HybridSearchIndex.h
+++ b/Engine/Plugins/Experimental/SemanticSearch/Source/SemanticSearch/Public/HybridSearchIndex.h
@@ -97,6 +97,9 @@ public:
 	/** Access the BM25 index directly (e.g. for diagnostics). */
 	FBM25Index& GetBM25Index();
 
+	/** Reset BM25 and failed sets. Used when re-indexing all assets from scratch. */
+	SEMANTICSEARCH_API void ResetBM25AndFailedState();
+
 	/** Number of assets currently in the vector index. */
 	SEMANTICSEARCH_API int32 GetIndexedCount() const;
 
@@ -106,6 +109,9 @@ public:
 	/** Number of assets that failed to index. */
 	SEMANTICSEARCH_API int32 GetFailedCount() const;
 
+	/** Check if a specific asset is in the failed set. */
+	SEMANTICSEARCH_API bool IsFailed(int64 AssetID) const;
+
 	/**
 	 * Save the full index (vector + BM25 + package hashes) to a file.
 	 * @param FilePath		Destination file path.
diff --git a/Engine/Plugins/Experimental/SemanticSearch/Source/SemanticSearch/Public/ISemanticSearchModule.h b/Engine/Plugins/Experimental/SemanticSearch/Source/SemanticSearch/Public/ISemanticSearchModule.h
--- a/Engine/Plugins/Experimental/SemanticSearch/Source/SemanticSearch/Public/ISemanticSearchModule.h
+++ b/Engine/Plugins/Experimental/SemanticSearch/Source/SemanticSearch/Public/ISemanticSearchModule.h
@@ -23,6 +23,7 @@ namespace UE::SemanticSearch
 		bool bIsTrained = false;
 		int32 SupportedAssetCount = 0;
 		int32 FailedCount = 0;
+		bool bIsIndexing = false;
 	};
 
 	class ISemanticSearchModule : public IModuleInterface
@@ -35,10 +36,9 @@ namespace UE::SemanticSearch
 		}
 
 		virtual void RegisterEmbeddingProvider(TUniquePtr<IEmbeddingProvider>&& EmbeddingProvider) = 0;
-
+		
 		/** Returns the active embedding provider, or nullptr if none has been registered. */
 		virtual TSharedPtr<IEmbeddingProvider> GetEmbeddingProvider() const = 0;
-
 		/** Index a single asset, always building from the backend if not already cached. */
 		virtual void IndexAsset(const FAssetData& AssetData) = 0;
 
@@ -53,6 +53,9 @@ namespace UE::SemanticSearch
 		/** Switch to a different index type. Re-indexes all assets. */
 		virtual void SwitchIndexType(ESemanticSearchIndexType NewType) = 0;
 
+		/** Re-index all supported assets. When bForceBuild is true, clears the index and always builds on DDC miss. */
+		virtual void IndexAllAssets(bool bForceBuild = false) = 0;
+
 		/** Retrain the current quantized index from scratch (new codebook). No-op for flat index. */
 		virtual void RetrainIndex() = 0;
 	};
diff --git a/Engine/Plugins/Experimental/SemanticSearch/Source/SemanticSearch/Public/Settings/SemanticSearchSettings.h b/Engine/Plugins/Experimental/SemanticSearch/Source/SemanticSearch/Public/Settings/SemanticSearchSettings.h
--- a/Engine/Plugins/Experimental/SemanticSearch/Source/SemanticSearch/Public/Settings/SemanticSearchSettings.h
+++ b/Engine/Plugins/Experimental/SemanticSearch/Source/SemanticSearch/Public/Settings/SemanticSearchSettings.h
@@ -12,8 +12,8 @@ enum class ESemanticSearchIndexType : uint8
 	PQ      UMETA(DisplayName = "PQ (Compressed)")
 };
 
-UCLASS(config = EditorPerProjectUserSettings, DefaultConfig, meta = (DisplayName = "Semantic Search"))
-class SEMANTICSEARCH_API USemanticSearchSettings : public UDeveloperSettings
+UCLASS(config = EditorPerProjectUserSettings, DefaultConfig, MinimalAPI, meta = (DisplayName = "Semantic Search"))
+class USemanticSearchSettings : public UDeveloperSettings
 {
 	GENERATED_BODY()
 
@@ -27,6 +27,8 @@ public:
 	virtual FName GetCategoryName() const override { return TEXT("Plugins"); }
 	virtual FName GetSectionName() const override { return TEXT("Semantic Search"); }
 
+	SEMANTICSEARCH_API virtual void PostEditChangeProperty(FPropertyChangedEvent& PropertyChangedEvent) override;
+
 	UPROPERTY(EditAnywhere, Config, Category = "Indexing",
 		meta = (ToolTip = "Content browser paths to include in the semantic search index.", RelativeToGameContentDir))
 	TArray<FDirectoryPath> IndexedPaths = { {TEXT("/Game")} };
@@ -41,7 +43,7 @@ public:
 
 	UPROPERTY(VisibleAnywhere, Config, Category = "Vector Index",
 		meta = (ToolTip = "Index type is managed via Tools > Semantic Search."))
-	ESemanticSearchIndexType IndexType = ESemanticSearchIndexType::Flat;
+	ESemanticSearchIndexType IndexType = ESemanticSearchIndexType::PQ;
 
 	UPROPERTY(EditAnywhere, Config, Category = "Vector Index",
 		meta = (EditCondition = "IndexType == ESemanticSearchIndexType::PQ", EditConditionHides,
diff --git a/Engine/Plugins/Experimental/SemanticSearch/Source/SemanticSearchEditorIntegrations/Private/ContentBrowser/SemanticSearchCB.cpp b/Engine/Plugins/Experimental/SemanticSearch/Source/SemanticSearchEditorIntegrations/Private/ContentBrowser/SemanticSearchCB.cpp
--- a/Engine/Plugins/Experimental/SemanticSearch/Source/SemanticSearchEditorIntegrations/Private/ContentBrowser/SemanticSearchCB.cpp
+++ b/Engine/Plugins/Experimental/SemanticSearch/Source/SemanticSearchEditorIntegrations/Private/ContentBrowser/SemanticSearchCB.cpp
@@ -20,6 +20,7 @@
 #include "Styling/SlateColor.h"
 #include "Templates/SharedPointer.h"
 #include "ToolMenus.h"
+#include "Widgets/Images/SImage.h"
 #include "Widgets/Input/SButton.h"
 #include "Widgets/Layout/SBox.h"
 #include "Widgets/SBoxPanel.h"
@@ -387,6 +388,20 @@ static FDelayedAutoRegisterHelper SemanticSearchStatsNavBarExtension(
 						return FMath::Max(0, Stats.SupportedAssetCount - Stats.VectorCount - Stats.FailedCount);
 					};
 
+					auto IsNotIndexing = []() -> bool
+					{
+						return !UE::SemanticSearch::ISemanticSearchModule::Get().GetIndexStats().bIsIndexing;
+					};
+
+					auto MakeRefreshButton = [](const FSlateBrush* Icon) -> TSharedRef<SWidget>
+					{
+						return SNew(SImage)
+							.Image(Icon)
+							.DesiredSizeOverride(FVector2D(12.f, 12.f));
+					};
+
+					const FSlateBrush* RefreshIcon = FAppStyle::GetBrush("SourceControl.Actions.Refresh");
+
 					FToolMenuEntry& Entry = InSection.AddEntry(
 						FToolMenuEntry::InitWidget(
 							"SemanticSearchStatsEntry",
@@ -400,74 +415,177 @@ static FDelayedAutoRegisterHelper SemanticSearchStatsNavBarExtension(
 							})
 							[
 								SNew(SHorizontalBox)
+
+								// "Semantic Search Index:"
 								+ SHorizontalBox::Slot()
 								.AutoWidth()
 								.VAlign(VAlign_Center)
 								[
 									SNew(STextBlock)
 									.TextStyle(FAppStyle::Get(), "SmallText")
-									.Text_Lambda([]() -> FText
-									{
-										const UE::SemanticSearch::FSemanticSearchIndexStats Stats = UE::SemanticSearch::ISemanticSearchModule::Get().GetIndexStats();
-										const int32 Indexed = Stats.VectorCount;
-										const int32 Failed = Stats.FailedCount;
-										const int32 Unindexed = FMath::Max(0, Stats.SupportedAssetCount - Indexed - Failed);
-										if (Failed > 0)
+									.Text(NSLOCTEXT("SemanticSearch", "IndexLabel", "Semantic Search Index:"))
+								]
+
+								// "{N} indexed" + reindex-all refresh icon
+								+ SHorizontalBox::Slot()
+								.AutoWidth()
+								.VAlign(VAlign_Center)
+								.Padding(4.f, 0.f, 0.f, 0.f)
+								[
+									SNew(SHorizontalBox)
+									+ SHorizontalBox::Slot()
+									.AutoWidth()
+									.VAlign(VAlign_Center)
+									[
+										SNew(STextBlock)
+										.TextStyle(FAppStyle::Get(), "SmallText")
+										.Text_Lambda([]() -> FText
 										{
-											return FText::Format(
-												NSLOCTEXT("SemanticSearch", "IndexStatsWithFailed",
-													"Semantic Search Index: {0} indexed | {1} unindexed | {2} failed"),
-												FText::AsNumber(Indexed),
-												FText::AsNumber(Unindexed),
-												FText::AsNumber(Failed));
-										}
-										return FText::Format(
-											NSLOCTEXT("SemanticSearch", "IndexStats",
-												"Semantic Search Index: {0} indexed | {1} unindexed"),
-											FText::AsNumber(Indexed),
-											FText::AsNumber(Unindexed));
-									})
+											const int32 Count = UE::SemanticSearch::ISemanticSearchModule::Get().GetIndexStats().VectorCount;
+											return FText::Format(NSLOCTEXT("SemanticSearch", "IndexedCount", "{0} indexed"), FText::AsNumber(Count));
+										})
+									]
+									+ SHorizontalBox::Slot()
+									.AutoWidth()
+									.VAlign(VAlign_Center)
+									.Padding(2.f, 0.f, 0.f, 0.f)
+									[
+										SNew(SButton)
+										.ButtonStyle(FAppStyle::Get(), "HoverOnlyButton")
+										.ContentPadding(FMargin(2.f, 0.f))
+										.ToolTipText(NSLOCTEXT("SemanticSearch", "ReindexAllTooltip", "Re-index all assets from scratch"))
+										.IsEnabled_Lambda(IsNotIndexing)
+										.OnClicked_Lambda([]()
+										{
+											UE::SemanticSearch::ISemanticSearchModule::Get().IndexAllAssets(/*bForceBuild=*/true);
+											return FReply::Handled();
+										})
+										[
+											MakeRefreshButton(RefreshIcon)
+										]
+									]
 								]
+
+								// " | {N} unindexed" + index-unindexed refresh icon
 								+ SHorizontalBox::Slot()
 								.AutoWidth()
 								.VAlign(VAlign_Center)
-								.Padding(6.f, 0.f, 0.f, 0.f)
+								.Padding(4.f, 0.f, 0.f, 0.f)
 								[
-									SNew(SButton)
-									.ButtonStyle(FAppStyle::Get(), "HoverOnlyButton")
-									.ContentPadding(FMargin(4.f, 0.f))
-									.ToolTipText(NSLOCTEXT("SemanticSearch", "ReindexTooltip", "Index all unindexed assets in configured folders"))
+									SNew(SHorizontalBox)
 									.Visibility_Lambda([GetUnindexedCount]()
 									{
 										return GetUnindexedCount() > 0
 											? EVisibility::Visible
 											: EVisibility::Collapsed;
 									})
-									.OnClicked_Lambda([]()
-									{
-										UE::SemanticSearch::FHybridSearchIndex& Index = UE::SemanticSearch::FHybridSearchIndex::Get();
-										UE::SemanticSearch::FAssetProcessorManager& ProcessorManager = UE::SemanticSearch::FAssetProcessorManager::Get();
-										UE::SemanticSearch::ISemanticSearchModule& Module = UE::SemanticSearch::ISemanticSearchModule::Get();
-										IAssetRegistry& AR = IAssetRegistry::GetChecked();
-
-										AR.EnumerateAllAssets([&](const FAssetData& Asset)
+									+ SHorizontalBox::Slot()
+									.AutoWidth()
+									.VAlign(VAlign_Center)
+									[
+										SNew(STextBlock)
+										.TextStyle(FAppStyle::Get(), "SmallText")
+										.Text_Lambda([]() -> FText
+										{
+											const UE::SemanticSearch::FSemanticSearchIndexStats Stats = UE::SemanticSearch::ISemanticSearchModule::Get().GetIndexStats();
+											const int32 Unindexed = FMath::Max(0, Stats.SupportedAssetCount - Stats.VectorCount - Stats.FailedCount);
+											return FText::Format(NSLOCTEXT("SemanticSearch", "UnindexedCount", "| {0} unindexed"), FText::AsNumber(Unindexed));
+										})
+									]
+									+ SHorizontalBox::Slot()
+									.AutoWidth()
+									.VAlign(VAlign_Center)
+									.Padding(2.f, 0.f, 0.f, 0.f)
+									[
+										SNew(SButton)
+										.ButtonStyle(FAppStyle::Get(), "HoverOnlyButton")
+										.ContentPadding(FMargin(2.f, 0.f))
+										.ToolTipText(NSLOCTEXT("SemanticSearch", "IndexUnindexedTooltip", "Index all unindexed assets"))
+										.IsEnabled_Lambda(IsNotIndexing)
+										.OnClicked_Lambda([]()
 										{
-											if (!Asset.IsRedirector()
-												&& UE::SemanticSearch::ISemanticSearchModule::IsInIndexedFolder(Asset)
-												&& ProcessorManager.GetProcessorForAsset(Asset)
-												&& !Index.Contains(Asset))
+											UE::SemanticSearch::FHybridSearchIndex& Index = UE::SemanticSearch::FHybridSearchIndex::Get();
+											UE::SemanticSearch::FAssetProcessorManager& ProcessorManager = UE::SemanticSearch::FAssetProcessorManager::Get();
+											UE::SemanticSearch::ISemanticSearchModule& Module = UE::SemanticSearch::ISemanticSearchModule::Get();
+											IAssetRegistry& AR = IAssetRegistry::GetChecked();
+
+											AR.EnumerateAllAssets([&](const FAssetData& Asset)
 											{
-												Module.IndexAsset(Asset);
-											}
-											return true;
-										});
-										return FReply::Handled();
+												if (!Asset.IsRedirector()
+													&& UE::SemanticSearch::ISemanticSearchModule::IsInIndexedFolder(Asset)
+													&& ProcessorManager.GetProcessorForAsset(Asset)
+													&& !Index.Contains(Asset))
+												{
+													Module.IndexAsset(Asset);
+												}
+												return true;
+											});
+											return FReply::Handled();
+										})
+										[
+											MakeRefreshButton(RefreshIcon)
+										]
+									]
+								]
+
+								// " | {N} failed" + retry-failed refresh icon
+								+ SHorizontalBox::Slot()
+								.AutoWidth()
+								.VAlign(VAlign_Center)
+								.Padding(4.f, 0.f, 0.f, 0.f)
+								[
+									SNew(SHorizontalBox)
+									.Visibility_Lambda([]()
+									{
+										return UE::SemanticSearch::FHybridSearchIndex::Get().GetFailedCount() > 0
+											? EVisibility::Visible
+											: EVisibility::Collapsed;
 									})
+									+ SHorizontalBox::Slot()
+									.AutoWidth()
+									.VAlign(VAlign_Center)
 									[
 										SNew(STextBlock)
 										.TextStyle(FAppStyle::Get(), "SmallText")
-										.Text(NSLOCTEXT("SemanticSearch", "ReindexButton", "Reindex"))
-										.ColorAndOpacity(FSlateColor(FLinearColor(0.f, 0.8f, 1.f)))
+										.Text_Lambda([]() -> FText
+										{
+											const int32 Failed = UE::SemanticSearch::FHybridSearchIndex::Get().GetFailedCount();
+											return FText::Format(NSLOCTEXT("SemanticSearch", "FailedCount", "| {0} failed"), FText::AsNumber(Failed));
+										})
+									]
+									+ SHorizontalBox::Slot()
+									.AutoWidth()
+									.VAlign(VAlign_Center)
+									.Padding(2.f, 0.f, 0.f, 0.f)
+									[
+										SNew(SButton)
+										.ButtonStyle(FAppStyle::Get(), "HoverOnlyButton")
+										.ContentPadding(FMargin(2.f, 0.f))
+										.ToolTipText(NSLOCTEXT("SemanticSearch", "RetryFailedTooltip", "Retry indexing all failed assets"))
+										.IsEnabled_Lambda(IsNotIndexing)
+										.OnClicked_Lambda([]()
+										{
+											UE::SemanticSearch::FHybridSearchIndex& Index = UE::SemanticSearch::FHybridSearchIndex::Get();
+											UE::SemanticSearch::FAssetProcessorManager& ProcessorManager = UE::SemanticSearch::FAssetProcessorManager::Get();
+											UE::SemanticSearch::ISemanticSearchModule& Module = UE::SemanticSearch::ISemanticSearchModule::Get();
+											IAssetRegistry& AR = IAssetRegistry::GetChecked();
+
+											AR.EnumerateAllAssets([&](const FAssetData& Asset)
+											{
+												if (!Asset.IsRedirector()
+													&& UE::SemanticSearch::ISemanticSearchModule::IsInIndexedFolder(Asset)
+													&& ProcessorManager.GetProcessorForAsset(Asset)
+													&& Index.IsFailed(UE::SemanticSearch::GetAssetIndexID(Asset)))
+												{
+													Module.IndexAsset(Asset);
+												}
+												return true;
+											});
+											return FReply::Handled();
+										})
+										[
+											MakeRefreshButton(RefreshIcon)
+										]
 									]
 								]
 							],
