Add support for blueprints (just thumbnail based for now)

#Synced-CL 52578922

diff --git a/Engine/Plugins/Experimental/SemanticSearch/Source/SemanticSearch/Private/AssetProcessors/AssetProcessorUtils.cpp b/Engine/Plugins/Experimental/SemanticSearch/Source/SemanticSearch/Private/AssetProcessors/AssetProcessorUtils.cpp
--- a/Engine/Plugins/Experimental/SemanticSearch/Source/SemanticSearch/Private/AssetProcessors/AssetProcessorUtils.cpp
+++ b/Engine/Plugins/Experimental/SemanticSearch/Source/SemanticSearch/Private/AssetProcessors/AssetProcessorUtils.cpp
@@ -7,6 +7,7 @@
 #include "AssetProcessors/MaterialInstanceAssetProcessor.h"
 #include "AssetProcessors/SkeletalMeshAssetProcessor.h"
 #include "AssetProcessors/StaticMeshAssetProcessor.h"
+#include "AssetProcessors/BlueprintAssetProcessor.h"
 #include "AssetProcessors/TextureAssetProcessor.h"
 
 #include "AssetRegistry/AssetData.h"
@@ -67,6 +68,7 @@ void RegisterDefaultAssetProcessors()
 	Manager.RegisterAssetProcessor<FStaticMeshProcessor>();
 	Manager.RegisterAssetProcessor<FSkeletalMeshProcessor>();
 	Manager.RegisterAssetProcessor<FTextureProcessor>();
+	Manager.RegisterAssetProcessor<FBlueprintProcessor>();
 }
 
 }
diff --git a/Engine/Plugins/Experimental/SemanticSearch/Source/SemanticSearch/Private/AssetProcessors/BlueprintAssetProcessor.cpp b/Engine/Plugins/Experimental/SemanticSearch/Source/SemanticSearch/Private/AssetProcessors/BlueprintAssetProcessor.cpp
new file mode 100644
--- /dev/null
+++ b/Engine/Plugins/Experimental/SemanticSearch/Source/SemanticSearch/Private/AssetProcessors/BlueprintAssetProcessor.cpp
@@ -0,0 +1,42 @@
+// Copyright Epic Games, Inc. All Rights Reserved.
+
+#include "AssetProcessors/BlueprintAssetProcessor.h"
+
+#include "AssetProcessors/AssetProcessorUtils.h"
+
+#include "AssetRegistry/AssetData.h"
+#include "Blueprint/BlueprintSupport.h"
+#include "Engine/Blueprint.h"
+
+namespace UE::SemanticSearch::Private
+{
+
+UClass& FBlueprintProcessor::GetSupportedClass() const
+{
+	return *UBlueprint::StaticClass();
+}
+
+bool FBlueprintProcessor::SupportDerivedClasses() const
+{
+	return true;
+}
+
+TSharedPtr<FJsonObject> FBlueprintProcessor::GetMetadata(const TSharedRef<const FAssetData>& InAsset) const
+{
+	TSharedPtr<FJsonObject> Metadata = MakeShared<FJsonObject>();
+
+	SetMetadata(Metadata, InAsset, FBlueprintTags::ParentClassPath, TEXTVIEW("Parent Class"));
+	SetMetadata(Metadata, InAsset, FBlueprintTags::NativeParentClassPath, TEXTVIEW("Native Parent Class"));
+	SetMetadata(Metadata, InAsset, FBlueprintTags::BlueprintType, TEXTVIEW("Blueprint Type"));
+	SetMetadata(Metadata, InAsset, FBlueprintTags::BlueprintDescription, TEXTVIEW("Description"));
+	SetMetadata(Metadata, InAsset, FBlueprintTags::BlueprintCategory, TEXTVIEW("Category"));
+	SetMetadata(Metadata, InAsset, FBlueprintTags::IsDataOnly, TEXTVIEW("Is Data Only"));
+	SetMetadata(Metadata, InAsset, FBlueprintTags::NumReplicatedProperties, TEXTVIEW("Replicated Properties"));
+	SetMetadata(Metadata, InAsset, FBlueprintTags::NumNativeComponents, TEXTVIEW("Native Components"));
+	SetMetadata(Metadata, InAsset, FBlueprintTags::NumBlueprintComponents, TEXTVIEW("Blueprint Components"));
+	SetMetadata(Metadata, InAsset, FBlueprintTags::ImplementedInterfaces, TEXTVIEW("Implemented Interfaces"));
+
+	return Metadata;
+}
+
+}
diff --git a/Engine/Plugins/Experimental/SemanticSearch/Source/SemanticSearch/Private/AssetProcessors/BlueprintAssetProcessor.h b/Engine/Plugins/Experimental/SemanticSearch/Source/SemanticSearch/Private/AssetProcessors/BlueprintAssetProcessor.h
new file mode 100644
--- /dev/null
+++ b/Engine/Plugins/Experimental/SemanticSearch/Source/SemanticSearch/Private/AssetProcessors/BlueprintAssetProcessor.h
@@ -0,0 +1,25 @@
+// Copyright Epic Games, Inc. All Rights Reserved.
+
+#pragma once
+
+#include "Implementations/ThumbnailBaseProcessor.h"
+
+namespace UE::SemanticSearch::Private
+{
+
+class FBlueprintProcessor : public FThumbnailBaseAssetProcessor
+{
+public:
+	constexpr virtual FStringView GetProcessSubBucketName() const override
+	{
+		return TEXTVIEW("FBlueprintProcessor");
+	}
+
+	// Thread safe
+	virtual UClass& GetSupportedClass() const override;
+	virtual bool SupportDerivedClasses() const override;
+
+	virtual TSharedPtr<FJsonObject> GetMetadata(const TSharedRef<const FAssetData>& InAsset) const override;
+};
+
+}
